Move EvaluationNotifies plugin to Engine/Plugins/Experimental

diff --git a/Engine/Plugins/Experimental/EvaluationNotifies/EvaluationNotifies.uplugin b/Engine/Plugins/Experimental/EvaluationNotifies/EvaluationNotifies.uplugin
new file mode 100644
--- /dev/null
+++ b/Engine/Plugins/Experimental/EvaluationNotifies/EvaluationNotifies.uplugin
@@ -0,0 +1,33 @@
+{
+	"FileVersion": 1,
+	"Version": 1,
+	"VersionName": "0.1",
+	"FriendlyName": "Evaluation Notifies",
+	"Description": "A system for animation notifies which have animation evaluation time code.",
+	"Category": "Animation",
+	"CreatedBy": "Epic Games, Inc.",
+	"CreatedByURL": "https://epicgames.com",
+	"EnabledByDefault": false,
+	"CanContainContent": true,
+	"IsBetaVersion": false,
+	"IsExperimentalVersion": true,
+	"Installed": false,
+	"Modules": [
+		{
+			"Name": "EvaluationNotifiesRuntime",
+			"Type": "Runtime",
+			"LoadingPhase": "PreDefault"
+		},
+		{
+			"Name": "EvaluationNotifiesEditor",
+			"Type": "UncookedOnly",
+			"LoadingPhase": "PreDefault"
+		}
+	],
+	"Plugins": [
+		{
+			"Name": "AnimationWarping",
+			"Enabled" : true
+		}
+	]
+}
\ No newline at end of file
diff --git a/Engine/Plugins/Experimental/EvaluationNotifies/Source/Editor/EvaluationNotifiesEditor.Build.cs b/Engine/Plugins/Experimental/EvaluationNotifies/Source/Editor/EvaluationNotifiesEditor.Build.cs
new file mode 100644
--- /dev/null
+++ b/Engine/Plugins/Experimental/EvaluationNotifies/Source/Editor/EvaluationNotifiesEditor.Build.cs
@@ -0,0 +1,43 @@
+// Copyright Epic Games, Inc. All Rights Reserved.
+
+namespace UnrealBuildTool.Rules
+{
+	public class EvaluationNotifiesEditor : ModuleRules
+	{
+		public EvaluationNotifiesEditor(ReadOnlyTargetRules Target) : base(Target)
+		{
+			PublicDependencyModuleNames.AddRange(
+			new string[]
+			{
+				"AnimationCore",
+				"EvaluationNotifiesRuntime",
+                "AnimGraph",
+				"AnimGraphRuntime",
+				"AnimationModifiers",
+				"AnimationBlueprintLibrary",
+				"Core",
+                "CoreUObject",
+                "Engine",
+            });
+
+			PrivateDependencyModuleNames.AddRange(
+			new string[]
+			{
+                "SlateCore",
+            });
+
+            if (Target.bBuildEditor == true)
+            {
+                PrivateDependencyModuleNames.AddRange(
+                    new string[]
+                    {
+						"BlueprintGraph",
+						"EditorFramework",
+						"Kismet",
+                        "UnrealEd",
+                    }
+                );
+            }
+        }
+	}
+}
diff --git a/Engine/Plugins/Experimental/EvaluationNotifies/Source/Editor/Private/AnimGraph/AnimGraphNode_EvaluationNotifies.cpp b/Engine/Plugins/Experimental/EvaluationNotifies/Source/Editor/Private/AnimGraph/AnimGraphNode_EvaluationNotifies.cpp
new file mode 100644
--- /dev/null
+++ b/Engine/Plugins/Experimental/EvaluationNotifies/Source/Editor/Private/AnimGraph/AnimGraphNode_EvaluationNotifies.cpp
@@ -0,0 +1,70 @@
+// Copyright Epic Games, Inc. All Rights Reserved.
+
+#include "AnimGraph/AnimGraphNode_EvaluationNotifies.h"
+#include "Animation/AnimRootMotionProvider.h"
+#include "Kismet2/BlueprintEditorUtils.h"
+#include "Kismet2/CompilerResultsLog.h"
+#include "DetailCategoryBuilder.h"
+#include "DetailLayoutBuilder.h"
+#include "PropertyHandle.h"
+#include "ScopedTransaction.h"
+
+#define LOCTEXT_NAMESPACE "AnimationWarping"
+
+UAnimGraphNode_EvaluationNotifies::UAnimGraphNode_EvaluationNotifies(const FObjectInitializer& ObjectInitializer) : Super(ObjectInitializer)
+{
+}
+
+FText UAnimGraphNode_EvaluationNotifies::GetControllerDescription() const
+{
+	return LOCTEXT("EvaluationNotifies", "EvaluationNotifies");
+}
+
+FText UAnimGraphNode_EvaluationNotifies::GetTooltipText() const
+{
+	return LOCTEXT("EvaluationNotifiesTooltip", "Rotates the root and lower body by the specified angle, while counter rotating the upper body to maintain the forward facing direction.");
+}
+
+FText UAnimGraphNode_EvaluationNotifies::GetNodeTitle(ENodeTitleType::Type TitleType) const
+{
+	return GetControllerDescription();
+}
+
+FLinearColor UAnimGraphNode_EvaluationNotifies::GetNodeTitleColor() const
+{
+	return FLinearColor(FColor(153.f, 0.f, 0.f));
+}
+
+void UAnimGraphNode_EvaluationNotifies::CustomizePinData(UEdGraphPin* Pin, FName SourcePropertyName, int32 ArrayIndex) const
+{
+	Super::CustomizePinData(Pin, SourcePropertyName, ArrayIndex);
+}
+
+void UAnimGraphNode_EvaluationNotifies::CustomizeDetails(IDetailLayoutBuilder& DetailBuilder)
+{
+	DECLARE_SCOPE_HIERARCHICAL_COUNTER_FUNC()
+	Super::CustomizeDetails(DetailBuilder);
+}
+
+void UAnimGraphNode_EvaluationNotifies::PostEditChangeProperty(struct FPropertyChangedEvent& PropertyChangedEvent)
+{
+	DECLARE_SCOPE_HIERARCHICAL_COUNTER_FUNC()
+	Super::PostEditChangeProperty(PropertyChangedEvent);
+}
+
+void UAnimGraphNode_EvaluationNotifies::GetInputLinkAttributes(FNodeAttributeArray& OutAttributes) const
+{
+	OutAttributes.Add(UE::Anim::IAnimRootMotionProvider::AttributeName);
+}
+
+void UAnimGraphNode_EvaluationNotifies::GetOutputLinkAttributes(FNodeAttributeArray& OutAttributes) const
+{
+	OutAttributes.Add(UE::Anim::IAnimRootMotionProvider::AttributeName);
+}
+
+void UAnimGraphNode_EvaluationNotifies::ValidateAnimNodeDuringCompilation(USkeleton* ForSkeleton, FCompilerResultsLog& MessageLog)
+{
+	Super::ValidateAnimNodeDuringCompilation(ForSkeleton, MessageLog);
+}
+
+#undef LOCTEXT_NAMESPACE
diff --git a/Engine/Plugins/Experimental/EvaluationNotifies/Source/Editor/Private/EvaluationNotifiesEditorModule.cpp b/Engine/Plugins/Experimental/EvaluationNotifies/Source/Editor/Private/EvaluationNotifiesEditorModule.cpp
new file mode 100644
--- /dev/null
+++ b/Engine/Plugins/Experimental/EvaluationNotifies/Source/Editor/Private/EvaluationNotifiesEditorModule.cpp
@@ -0,0 +1,11 @@
+// Copyright Epic Games, Inc. All Rights Reserved.
+
+#include "EvaluationNotifiesEditorModule.h"
+#include "Modules/ModuleManager.h"
+
+class FEvaluationNotifiesEditorModule : public IEvaluationNotifiesEditorModule
+{
+};
+
+
+IMPLEMENT_MODULE( FEvaluationNotifiesEditorModule, EvaluationNotifiesEditor )
diff --git a/Engine/Plugins/Experimental/EvaluationNotifies/Source/Editor/Public/AnimGraph/AnimGraphNode_EvaluationNotifies.h b/Engine/Plugins/Experimental/EvaluationNotifies/Source/Editor/Public/AnimGraph/AnimGraphNode_EvaluationNotifies.h
new file mode 100644
--- /dev/null
+++ b/Engine/Plugins/Experimental/EvaluationNotifies/Source/Editor/Public/AnimGraph/AnimGraphNode_EvaluationNotifies.h
@@ -0,0 +1,39 @@
+// Copyright Epic Games, Inc. All Rights Reserved.
+
+#pragma once
+#include "AnimGraphNode_SkeletalControlBase.h"
+#include "EvaluationNotifies/AnimNode_EvaluationNotifies.h"
+#include "AnimGraphNode_EvaluationNotifies.generated.h"
+
+namespace ENodeTitleType { enum Type : int; }
+
+UCLASS()
+class EVALUATIONNOTIFIESEDITOR_API UAnimGraphNode_EvaluationNotifies : public UAnimGraphNode_SkeletalControlBase
+{
+	GENERATED_UCLASS_BODY()
+
+	UPROPERTY(EditAnywhere, Category = Settings)
+	FAnimNode_EvaluationNotifies Node;
+
+public:
+	// UEdGraphNode interface
+	virtual FText GetNodeTitle(ENodeTitleType::Type TitleType) const override;
+	virtual FLinearColor GetNodeTitleColor() const override;
+	virtual FText GetTooltipText() const override;
+	virtual void CustomizeDetails(IDetailLayoutBuilder& DetailBuilder) override;
+	virtual void PostEditChangeProperty(struct FPropertyChangedEvent& PropertyChangedEvent) override;
+	virtual void GetInputLinkAttributes(FNodeAttributeArray& OutAttributes) const override;
+	virtual void GetOutputLinkAttributes(FNodeAttributeArray& OutAttributes) const override;
+	virtual void ValidateAnimNodeDuringCompilation(USkeleton* ForSkeleton, FCompilerResultsLog& MessageLog) override;
+	// End of UEdGraphNode interface
+
+protected:
+	// UAnimGraphNode_Base interface
+	virtual void CustomizePinData(UEdGraphPin* Pin, FName SourcePropertyName, int32 ArrayIndex) const override;
+	// End of UAnimGraphNode_Base interface
+
+	// UAnimGraphNode_SkeletalControlBase interface
+	virtual FText GetControllerDescription() const override;
+	virtual const FAnimNode_SkeletalControlBase* GetNode() const override { return &Node; }
+	// End of UAnimGraphNode_SkeletalControlBase interface
+};
diff --git a/Engine/Plugins/Experimental/EvaluationNotifies/Source/Editor/Public/EvaluationNotifiesEditorModule.h b/Engine/Plugins/Experimental/EvaluationNotifies/Source/Editor/Public/EvaluationNotifiesEditorModule.h
new file mode 100644
--- /dev/null
+++ b/Engine/Plugins/Experimental/EvaluationNotifies/Source/Editor/Public/EvaluationNotifiesEditorModule.h
@@ -0,0 +1,9 @@
+// Copyright Epic Games, Inc. All Rights Reserved.
+
+#pragma once
+
+#include "Modules/ModuleInterface.h"
+
+class IEvaluationNotifiesEditorModule : public IModuleInterface
+{
+};
diff --git a/Engine/Plugins/Experimental/EvaluationNotifies/Source/Runtime/EvaluationNotifiesRuntime.Build.cs b/Engine/Plugins/Experimental/EvaluationNotifies/Source/Runtime/EvaluationNotifiesRuntime.Build.cs
new file mode 100644
--- /dev/null
+++ b/Engine/Plugins/Experimental/EvaluationNotifies/Source/Runtime/EvaluationNotifiesRuntime.Build.cs
@@ -0,0 +1,21 @@
+// Copyright Epic Games, Inc. All Rights Reserved.
+
+namespace UnrealBuildTool.Rules
+{
+	public class EvaluationNotifiesRuntime : ModuleRules
+	{
+		public EvaluationNotifiesRuntime(ReadOnlyTargetRules Target) : base(Target)
+		{
+			PublicDependencyModuleNames.AddRange(
+			new string[]
+			{
+				"AnimationCore",
+                "AnimGraphRuntime",
+                "Core",
+                "CoreUObject",
+                "Engine",
+                "AnimationWarpingRuntime"
+            });
+		}
+	}
+}
diff --git a/Engine/Plugins/Experimental/EvaluationNotifies/Source/Runtime/Private/EvaluationNotifies/AnimNode_EvaluationNotifies.cpp b/Engine/Plugins/Experimental/EvaluationNotifies/Source/Runtime/Private/EvaluationNotifies/AnimNode_EvaluationNotifies.cpp
new file mode 100644
--- /dev/null
+++ b/Engine/Plugins/Experimental/EvaluationNotifies/Source/Runtime/Private/EvaluationNotifies/AnimNode_EvaluationNotifies.cpp
@@ -0,0 +1,122 @@
+// Copyright Epic Games, Inc. All Rights Reserved.
+
+#include "EvaluationNotifies/AnimNode_EvaluationNotifies.h"
+#include "Animation/AnimInstanceProxy.h"
+#include "Animation/AnimRootMotionProvider.h"
+#include "Animation/AnimTrace.h"
+#include "BoneControllers/AnimNode_OffsetRootBone.h"
+#include "Curves/BezierUtilities.h"
+#include "StructUtils/InstancedStruct.h"
+
+#include UE_INLINE_GENERATED_CPP_BY_NAME(AnimNode_EvaluationNotifies)
+
+TMap<UClass*, UScriptStruct*> FAnimNode_EvaluationNotifies::NotifyEvaluationHandlerMap;
+
+void FAnimNode_EvaluationNotifies::GetEvaluationNotifiesFromAnimation(const UAnimSequenceBase* Animation, TArray<FInstancedStruct>& OutNotifyInstances)
+{
+	if(Animation)
+	{
+		// todo: we should probably preprocess sequences for this number
+		int32 Count = 0;
+		for (int32 Idx = 0; Idx < Animation->Notifies.Num(); Idx++)
+		{
+			const FAnimNotifyEvent& NotifyEvent = Animation->Notifies[Idx];
+			if (UAnimNotifyState* Notify = Cast<UAnimNotifyState>(NotifyEvent.NotifyStateClass))
+			{
+				if (UScriptStruct** EvaluationHandlerStruct = NotifyEvaluationHandlerMap.Find(Notify->GetClass()))
+				{
+					Count++;
+				}
+			}
+		}
+
+		OutNotifyInstances.Empty(Count);
+
+		for (int32 Idx = 0; Idx < Animation->Notifies.Num(); Idx++)
+		{
+			const FAnimNotifyEvent& NotifyEvent = Animation->Notifies[Idx];
+			if (UAnimNotifyState* Notify = Cast<UAnimNotifyState>(NotifyEvent.NotifyStateClass))
+			{
+				if (UScriptStruct** EvaluationHandlerStruct = NotifyEvaluationHandlerMap.Find(Notify->GetClass()))
+				{
+					FInstancedStruct& Struct = OutNotifyInstances.AddDefaulted_GetRef();
+					Struct.InitializeAs(*EvaluationHandlerStruct);
+					FEvaluationNotifyInstance& Instance = Struct.GetMutable<FEvaluationNotifyInstance>();
+					Instance.AnimNotify = Notify;
+					Instance.StartTime = NotifyEvent.GetTriggerTime();
+					Instance.EndTime = NotifyEvent.GetEndTriggerTime();
+				}
+			}
+		}
+	}
+}
+
+void FAnimNode_EvaluationNotifies::UpdateInternal(const FAnimationUpdateContext& Context)
+{
+	FAnimNode_SkeletalControlBase::UpdateInternal(Context);
+
+	if (UE::AnimationWarping::FRootOffsetProvider* RootOffsetProvider = Context.GetMessage<UE::AnimationWarping::FRootOffsetProvider>())
+	{
+		RootBoneTransform = RootOffsetProvider->GetRootTransform();
+	}
+	else
+	{
+		RootBoneTransform = Context.AnimInstanceProxy->GetComponentTransform();
+	}
+}
+
+void FAnimNode_EvaluationNotifies::Initialize_AnyThread(const FAnimationInitializeContext& Context)
+{
+	FAnimNode_SkeletalControlBase::Initialize_AnyThread(Context);
+}
+
+
+void FAnimNode_EvaluationNotifies::EvaluateSkeletalControl_AnyThread(FComponentSpacePoseContext& Output, TArray<FBoneTransform>& OutBoneTransforms)
+{
+	const float DeltaSeconds = Output.AnimInstanceProxy->GetDeltaSeconds();
+
+	if (DeltaSeconds > 0.f)
+	{
+		const UAnimSequenceBase* AnimSequenceBase = Cast<UAnimSequenceBase>(CurrentAnimAsset);
+		if (AnimSequenceBase != CurrentSequence)
+		{
+			PreviousAnimAssetTime = CurrentAnimAssetTime;
+			CurrentSequence = AnimSequenceBase;
+			GetEvaluationNotifiesFromAnimation(CurrentSequence, Tags);
+		}
+
+		for(auto& TagData : Tags)
+		{
+			FEvaluationNotifyInstance& Tag = TagData.GetMutable<FEvaluationNotifyInstance>();
+			if (Tag.StartTime <= CurrentAnimAssetTime && Tag.EndTime >= CurrentAnimAssetTime)
+			{
+				if (!Tag.bActive)
+				{
+					Tag.bActive = true;
+					Tag.Start(CurrentSequence);
+				}
+
+				Tag.Update(CurrentSequence, CurrentAnimAssetTime, DeltaSeconds, RootBoneTransform, NamedTransforms, Output, OutBoneTransforms);
+			}
+			else
+			{
+				if (Tag.bActive)
+				{
+					Tag.bActive = false;
+					Tag.End();
+				}
+			}
+			
+		}
+	}
+}
+
+void FAnimNode_EvaluationNotifies::RegisterEvaluationHandler(UClass* NotifyType, UScriptStruct* HandlerType)
+{
+	NotifyEvaluationHandlerMap.Add(NotifyType, HandlerType);
+}
+
+void FAnimNode_EvaluationNotifies::UnregisterEvaluationHandler(UClass* NotifyType)
+{
+	NotifyEvaluationHandlerMap.Remove(NotifyType);
+}
diff --git a/Engine/Plugins/Experimental/EvaluationNotifies/Source/Runtime/Private/EvaluationNotifies/AnimNotifyState_Alignment.cpp b/Engine/Plugins/Experimental/EvaluationNotifies/Source/Runtime/Private/EvaluationNotifies/AnimNotifyState_Alignment.cpp
new file mode 100644
--- /dev/null
+++ b/Engine/Plugins/Experimental/EvaluationNotifies/Source/Runtime/Private/EvaluationNotifies/AnimNotifyState_Alignment.cpp
@@ -0,0 +1,59 @@
+// Copyright Epic Games, Inc. All Rights Reserved.
+
+#include "EvaluationNotifies/AnimNotifyState_Alignment.h"
+#include "Animation/AnimInstanceProxy.h"
+#include "Animation/AnimRootMotionProvider.h"
+
+#include UE_INLINE_GENERATED_CPP_BY_NAME(AnimNotifyState_Alignment)
+
+void FAlignmentNotifyInstance::Start(const UAnimSequenceBase* AnimationAsset)
+{
+	UNotifyState_Alignment* AlignmentNotify = Cast<UNotifyState_Alignment>(AnimNotify);
+	AlignBone = AlignmentNotify->AlignBone;
+}
+
+void FAlignmentNotifyInstance::Update(const UAnimSequenceBase* AnimationAsset, float CurrentTime, float DeltaTime, 
+                                   FTransform& RootBoneTransform, const TMap<FName, FTransform>& NamedTransforms, FComponentSpacePoseContext& Output, TArray<FBoneTransform>& OutBoneTransforms)
+{
+	UNotifyState_Alignment* AlignmentNotify = Cast<UNotifyState_Alignment>(AnimNotify);
+	if (!AlignBone.HasValidSetup())
+	{
+		const FBoneContainer& RequiredBones = Output.AnimInstanceProxy->GetRequiredBones();	
+		AlignBone.Initialize(RequiredBones);
+	}
+
+	const UE::Anim::IAnimRootMotionProvider* RootMotionProvider = UE::Anim::IAnimRootMotionProvider::Get();
+    ensureMsgf(RootMotionProvider, TEXT("Alignment expected a valid root motion delta provider interface."));
+
+    if (RootMotionProvider)
+    {
+    	FTransform ThisFrameRootMotionTransform;
+    	if (RootMotionProvider->ExtractRootMotion(Output.CustomAttributes, ThisFrameRootMotionTransform))
+    	{
+    		if (const FTransform* TargetTransform = NamedTransforms.Find(AlignmentNotify->TransformName))
+    		{
+    			// simple linear blend weight
+    			float TimeRemaining = EndTime-CurrentTime;
+    			float Weight = TimeRemaining <= 0.f ? 1.f : FMath::Min(1.f, DeltaTime / TimeRemaining);
+    		
+				const FBoneContainer& RequiredBones = Output.AnimInstanceProxy->GetRequiredBones();	
+    			FCompactPoseBoneIndex AlignBoneIndex = AlignBone.GetCompactPoseIndex(RequiredBones);
+    			if (AlignBoneIndex.IsValid())
+    			{
+    				const FTransform& AlignBoneTransform  = Output.Pose.GetComponentSpaceTransform(AlignBoneIndex);
+    				RootBoneTransform = AlignBoneTransform * ThisFrameRootMotionTransform * RootBoneTransform;;
+    			}
+
+    			RootBoneTransform = AlignmentNotify->AlignOffset * RootBoneTransform;
+    		
+    			FTransform RelativeOffset = TargetTransform->GetRelativeTransform(RootBoneTransform);
+
+    			FTransform ScaledOffset = FTransform(FMath::Lerp(FQuat::Identity, RelativeOffset.GetRotation(), Weight),
+														Weight * RelativeOffset.GetTranslation());
+
+    			RootMotionProvider->OverrideRootMotion(ThisFrameRootMotionTransform * ScaledOffset, Output.CustomAttributes);
+    		}
+    	}
+    }
+}
+
diff --git a/Engine/Plugins/Experimental/EvaluationNotifies/Source/Runtime/Private/EvaluationNotifies/AnimNotifyState_TwoBoneIK.cpp b/Engine/Plugins/Experimental/EvaluationNotifies/Source/Runtime/Private/EvaluationNotifies/AnimNotifyState_TwoBoneIK.cpp
new file mode 100644
--- /dev/null
+++ b/Engine/Plugins/Experimental/EvaluationNotifies/Source/Runtime/Private/EvaluationNotifies/AnimNotifyState_TwoBoneIK.cpp
@@ -0,0 +1,240 @@
+// Copyright Epic Games, Inc. All Rights Reserved.
+
+#include "EvaluationNotifies/AnimNotifyState_TwoBoneIK.h"
+#include "Animation/AnimInstanceProxy.h"
+#include "Animation/AnimRootMotionProvider.h"
+#include "HAL/IConsoleManager.h"
+#include "TwoBoneIK.h"
+#include "VisualLogger/VisualLogger.h"
+
+#include UE_INLINE_GENERATED_CPP_BY_NAME(AnimNotifyState_TwoBoneIK)
+
+void FTwoBoneIKNotifyInstance::Start(const UAnimSequenceBase* AnimationAsset)
+{
+	UNotifyState_TwoBoneIK* TwoBoneIKNotify = Cast<UNotifyState_TwoBoneIK>(AnimNotify);
+	IKBone = TwoBoneIKNotify->IKBone;
+	RelativeToBone = TwoBoneIKNotify->RelativeToBone;
+	EffectorTarget = TwoBoneIKNotify->EffectorTarget;
+	JointTarget = TwoBoneIKNotify->JointTarget;
+}
+
+FTransform GetTargetTransform(const FTransform& InComponentTransform, FCSPose<FCompactPose>& MeshBases, FBoneSocketTarget& InTarget, EBoneControlSpace Space, const FVector& InOffset) 
+{
+	FTransform OutTransform;
+	if (Space == BCS_BoneSpace)
+	{
+		OutTransform = InTarget.GetTargetTransform(InOffset, MeshBases, InComponentTransform);
+	}
+	else
+	{
+		// parent bone space still goes through this way
+		// if your target is socket, it will try find parents of joint that socket belongs to
+		OutTransform.SetLocation(InOffset);
+		FAnimationRuntime::ConvertBoneSpaceTransformToCS(InComponentTransform, MeshBases, OutTransform, InTarget.GetCompactPoseBoneIndex(), Space);
+	}
+
+	return OutTransform;
+}
+
+bool g_bEnableTwoBoneIKNotify = true;
+static FAutoConsoleVariableRef CVarEnableTwoBoneIKNotify(
+	TEXT("Animation.Notify.TwoBoneIK.Enable"),
+	g_bEnableTwoBoneIKNotify,
+	TEXT("Enable Two BoneIK Notify"));
+
+void FTwoBoneIKNotifyInstance::Update(const UAnimSequenceBase* AnimationAsset, float CurrentTime, float DeltaTime, 
+                                   FTransform& InRootBoneTransform, const TMap<FName, FTransform>& NamedTransforms, FComponentSpacePoseContext& Output, TArray<FBoneTransform>& OutBoneTransforms)
+{
+	if (!g_bEnableTwoBoneIKNotify)
+	{
+		return;
+	}
+	
+	const FBoneContainer& BoneContainer = Output.Pose.GetPose().GetBoneContainer();
+	
+	const UE::Anim::IAnimRootMotionProvider* RootMotionProvider = UE::Anim::IAnimRootMotionProvider::Get();
+    ensureMsgf(RootMotionProvider, TEXT("Alignment expected a valid root motion delta provider interface."));
+
+	FTransform RootBoneTransform = InRootBoneTransform;
+	FTransform ThisFrameRootMotionTransform;
+	if (RootMotionProvider->ExtractRootMotion(Output.CustomAttributes, ThisFrameRootMotionTransform))
+	{
+		RootBoneTransform =  ThisFrameRootMotionTransform * RootBoneTransform;
+	}
+	
+	UNotifyState_TwoBoneIK* Data = Cast<UNotifyState_TwoBoneIK>(AnimNotify);
+	check(Data);
+
+	if (!IKBone.IsValidToEvaluate())
+	{
+		IKBone.Initialize(BoneContainer);
+		RelativeToBone.Initialize(BoneContainer);
+
+		EffectorTarget.InitializeBoneReferences(BoneContainer);
+		JointTarget.InitializeBoneReferences(BoneContainer);
+
+		FCompactPoseBoneIndex IKBoneCompactPoseIndex = IKBone.GetCompactPoseIndex(BoneContainer);
+		CachedLowerLimbIndex = FCompactPoseBoneIndex(INDEX_NONE);
+		CachedUpperLimbIndex = FCompactPoseBoneIndex(INDEX_NONE);
+		if (IKBoneCompactPoseIndex != INDEX_NONE)
+		{
+			CachedLowerLimbIndex = BoneContainer.GetParentBoneIndex(IKBoneCompactPoseIndex);
+			if (CachedLowerLimbIndex != INDEX_NONE)
+			{
+				CachedUpperLimbIndex = BoneContainer.GetParentBoneIndex(CachedLowerLimbIndex);
+			}
+		}
+	}
+
+	// Get indices of the lower and upper limb bones and check validity.
+	bool bInvalidLimb = false;
+
+	FCompactPoseBoneIndex IKBoneCompactPoseIndex = IKBone.GetCompactPoseIndex(BoneContainer);
+	FCompactPoseBoneIndex RelativeToBoneCompactPoseIndex = RelativeToBone.GetCompactPoseIndex(BoneContainer);
+
+	const bool bInBoneSpace = (Data->EffectorLocationSpace == BCS_ParentBoneSpace) || (Data->EffectorLocationSpace == BCS_BoneSpace);
+
+	// Get Local Space transforms for our bones. We do this first in case they already are local.
+	// As right after we get them in component space. (And that does the auto conversion).
+	// We might save one transform by doing local first...
+	const FTransform EndBoneLocalTransform = Output.Pose.GetLocalSpaceTransform(IKBoneCompactPoseIndex);
+	const FTransform LowerLimbLocalTransform = Output.Pose.GetLocalSpaceTransform(CachedLowerLimbIndex);
+	const FTransform UpperLimbLocalTransform = Output.Pose.GetLocalSpaceTransform(CachedUpperLimbIndex);
+
+	// Now get those in component space...
+	FTransform LowerLimbCSTransform = Output.Pose.GetComponentSpaceTransform(CachedLowerLimbIndex);
+	FTransform UpperLimbCSTransform = Output.Pose.GetComponentSpaceTransform(CachedUpperLimbIndex);
+	FTransform EndBoneCSTransform = Output.Pose.GetComponentSpaceTransform(IKBoneCompactPoseIndex);
+
+	// Get current position of root of limb.
+	// All position are in Component space.
+	const FVector RootPos = UpperLimbCSTransform.GetTranslation();
+	const FVector InitialJointPos = LowerLimbCSTransform.GetTranslation();
+	const FVector InitialEndPos = EndBoneCSTransform.GetTranslation();
+
+	FVector EffectorLocation = Data->EffectorLocation;
+	if (const FTransform* EffectorLocationTransform = NamedTransforms.Find(Data->EffectorLocationTransformName))
+	{
+		EffectorLocation = EffectorLocationTransform->GetLocation();
+	}
+
+	// // Transform EffectorLocation from EffectorLocationSpace to ComponentSpace.
+	// FTransform EffectorTransform = GetTargetTransform(RootBoneTransform, Output.Pose, Data->EffectorTarget, Data->EffectorLocationSpace, EffectorLocation);
+	FTransform EffectorTransform;
+	EffectorTransform.SetLocation(EffectorLocation);
+	EffectorTransform = EffectorTransform.GetRelativeTransform(RootBoneTransform);
+
+	if(RelativeToBoneCompactPoseIndex.IsValid())
+	{
+		// if RelativeToBone was set, then compute relative position offset.
+		// - add component space position difference to effect position
+		
+		const FTransform RelativeToTransform = Output.Pose.GetComponentSpaceTransform(RelativeToBoneCompactPoseIndex);
+		const FVector Offset = EndBoneCSTransform.GetLocation() - RelativeToTransform.GetLocation();
+		EffectorTransform.SetLocation(EffectorTransform.GetLocation() + Offset);
+	}
+	
+	// blending.  For now just blending the effector location
+	float Weight = 1.0f;
+	if (Data->BlendInTime > 0.0f && CurrentTime - StartTime < Data->BlendInTime)
+	{
+		Weight = FMath::Max(CurrentTime - StartTime, 0)/Data->BlendInTime;
+	}
+	else if (Data->BlendOutTime > 0.0f && EndTime - StartTime < Data->BlendOutTime)
+	{
+		Weight = 1.0f - FMath::Max(EndTime - CurrentTime, 0)/Data->BlendInTime;
+	}
+
+	EffectorTransform.BlendWith(EndBoneCSTransform, 1.0f-Weight);
+
+	// Get joint target (used for defining plane that joint should be in).
+	FTransform JointTargetTransform = GetTargetTransform(RootBoneTransform, Output.Pose, Data->JointTarget, Data->JointTargetLocationSpace,  Data->JointTargetLocation);
+
+	FVector	JointTargetPos = JointTargetTransform.GetTranslation();
+
+	UE_VLOG_SPHERE(Output.AnimInstanceProxy->GetAnimInstanceObject(),"TwoBoneIK", Display, RootBoneTransform.TransformPosition(JointTargetPos), 0.1f, FColor::Red, TEXT(""));
+
+	// This is our reach goal.
+	FVector DesiredPos = EffectorTransform.GetTranslation();
+
+	// IK solver
+	UpperLimbCSTransform.SetLocation(RootPos);
+	LowerLimbCSTransform.SetLocation(InitialJointPos);
+	EndBoneCSTransform.SetLocation(InitialEndPos);
+
+	
+	AnimationCore::SolveTwoBoneIK(UpperLimbCSTransform, LowerLimbCSTransform, EndBoneCSTransform, JointTargetPos, DesiredPos, Data->bAllowStretching, Data->StartStretchRatio, Data->MaxStretchScale);
+
+// #if WITH_EDITOR
+// 	CachedJointTargetPos = JointTargetPos;
+// 	CachedJoints[0] = UpperLimbCSTransform.GetLocation();
+// 	CachedJoints[1] = LowerLimbCSTransform.GetLocation();
+// 	CachedJoints[2] = EndBoneCSTransform.GetLocation();
+// #endif // WITH_EDITOR
+
+	// if no twist, we clear twist from each limb
+	if (!Data->bAllowTwist)
+	{
+		auto RemoveTwist = [this](const FTransform& InParentTransform, FTransform& InOutTransform, const FTransform& OriginalLocalTransform, const FVector& InAlignVector) 
+		{
+			FTransform LocalTransform = InOutTransform.GetRelativeTransform(InParentTransform);
+			FQuat LocalRotation = LocalTransform.GetRotation();
+			FQuat NewTwist, NewSwing;
+			LocalRotation.ToSwingTwist(InAlignVector, NewSwing, NewTwist);
+			NewSwing.Normalize();
+
+			// get new twist from old local
+			LocalRotation = OriginalLocalTransform.GetRotation();
+			FQuat OldTwist, OldSwing;
+			LocalRotation.ToSwingTwist(InAlignVector, OldSwing, OldTwist);
+			OldTwist.Normalize();
+
+			InOutTransform.SetRotation(InParentTransform.GetRotation() * NewSwing * OldTwist);
+			InOutTransform.NormalizeRotation();
+		};
+
+		const FCompactPoseBoneIndex UpperLimbParentIndex = BoneContainer.GetParentBoneIndex(CachedUpperLimbIndex);
+		FVector AlignDir = Data->TwistAxis.GetTransformedAxis(FTransform::Identity);
+		if (UpperLimbParentIndex != INDEX_NONE)
+		{
+			FTransform UpperLimbParentTransform = Output.Pose.GetComponentSpaceTransform(UpperLimbParentIndex);
+			RemoveTwist(UpperLimbParentTransform, UpperLimbCSTransform, UpperLimbLocalTransform, AlignDir);
+		}
+			
+		RemoveTwist(UpperLimbCSTransform, LowerLimbCSTransform, LowerLimbLocalTransform, AlignDir);
+	}
+
+
+	
+	// Update transform for upper bone.
+	{
+		// Order important. First bone is upper limb.
+		OutBoneTransforms.Add( FBoneTransform(CachedUpperLimbIndex, UpperLimbCSTransform) );
+	}
+
+	// Update transform for lower bone.
+	{
+		// Order important. Second bone is lower limb.
+		OutBoneTransforms.Add( FBoneTransform(CachedLowerLimbIndex, LowerLimbCSTransform) );
+	}
+
+	// Update transform for end bone.
+	{
+		// only allow bTakeRotationFromEffectorSpace during bone space
+		if (bInBoneSpace && Data->bTakeRotationFromEffectorSpace)
+		{
+			EndBoneCSTransform.SetRotation(EffectorTransform.GetRotation());
+		}
+		else if (Data->bMaintainEffectorRelRot)
+		{
+			EndBoneCSTransform = EndBoneLocalTransform * LowerLimbCSTransform;
+		}
+		// Order important. Third bone is End Bone.
+		OutBoneTransforms.Add(FBoneTransform(IKBoneCompactPoseIndex, EndBoneCSTransform));
+		
+	}
+
+	// Make sure we have correct number of bones
+	// check(OutBoneTransforms.Num() == 3);
+}
+
diff --git a/Engine/Plugins/Experimental/EvaluationNotifies/Source/Runtime/Private/EvaluationNotifiesRuntimeModule.cpp b/Engine/Plugins/Experimental/EvaluationNotifies/Source/Runtime/Private/EvaluationNotifiesRuntimeModule.cpp
new file mode 100644
--- /dev/null
+++ b/Engine/Plugins/Experimental/EvaluationNotifies/Source/Runtime/Private/EvaluationNotifiesRuntimeModule.cpp
@@ -0,0 +1,48 @@
+// Copyright Epic Games, Inc. All Rights Reserved.
+
+#include "EvaluationNotifiesRuntimeModule.h"
+#include "Animation/AnimRootMotionProvider.h"
+#include "Features/IModularFeatures.h"
+#include "Modules/ModuleManager.h"
+#include "EvaluationNotifies/AnimNode_EvaluationNotifies.h"
+#include "EvaluationNotifies/AnimNotifyState_Alignment.h"
+#include "EvaluationNotifies/AnimNotifyState_TwoBoneIK.h"
+
+class FEvaluationNotifiesRuntimeModule : public IEvaluationNotifiesRuntimeModule
+{
+};
+
+namespace UE
+{
+	namespace EvaluationNotifies
+	{
+		class FModule : public IModuleInterface
+		{
+		public:
+
+			FModule();
+	
+			virtual void StartupModule() override;
+			virtual void ShutdownModule() override;
+		};
+
+		FModule::FModule()
+		{
+		}
+
+		void FModule::StartupModule()
+		{
+			FAnimNode_EvaluationNotifies::RegisterEvaluationHandler(UNotifyState_Alignment::StaticClass(), FAlignmentNotifyInstance::StaticStruct());
+			FAnimNode_EvaluationNotifies::RegisterEvaluationHandler(UNotifyState_TwoBoneIK::StaticClass(), FTwoBoneIKNotifyInstance::StaticStruct());
+		}
+
+		void FModule::ShutdownModule()
+		{
+			FAnimNode_EvaluationNotifies::UnregisterEvaluationHandler(UNotifyState_Alignment::StaticClass());
+			FAnimNode_EvaluationNotifies::UnregisterEvaluationHandler(UNotifyState_TwoBoneIK::StaticClass());
+		}
+	}
+}
+	
+
+IMPLEMENT_MODULE(UE::EvaluationNotifies::FModule, EvaluationNotifiesRuntime)
diff --git a/Engine/Plugins/Experimental/EvaluationNotifies/Source/Runtime/Public/EvaluationNotifies/AnimNode_EvaluationNotifies.h b/Engine/Plugins/Experimental/EvaluationNotifies/Source/Runtime/Public/EvaluationNotifies/AnimNode_EvaluationNotifies.h
new file mode 100644
--- /dev/null
+++ b/Engine/Plugins/Experimental/EvaluationNotifies/Source/Runtime/Public/EvaluationNotifies/AnimNode_EvaluationNotifies.h
@@ -0,0 +1,75 @@
+// Copyright Epic Games, Inc. All Rights Reserved.
+
+#pragma once
+
+#include "Animation/AnimNotifies/AnimNotifyState.h"
+#include "BoneControllers/BoneControllerTypes.h"
+#include "BoneControllers/AnimNode_SkeletalControlBase.h"
+#include "Kismet/KismetMathLibrary.h"
+#include "StructUtils/InstancedStruct.h"
+#include "AnimNode_EvaluationNotifies.generated.h"
+
+struct FAnimationInitializeContext;
+struct FComponentSpacePoseContext;
+struct FNodeDebugData;
+
+USTRUCT()
+struct FEvaluationNotifyInstance
+{
+	GENERATED_BODY()
+	
+	virtual void Start(const UAnimSequenceBase* AnimationAsset) {}
+	virtual void Update(const UAnimSequenceBase* AnimationAsset, float CurrentTime, float DeltaTime,
+		FTransform& RootBoneTransform, const TMap<FName, FTransform>& NamedTransforms, FComponentSpacePoseContext& Output, TArray<FBoneTransform>& OutBoneTransforms) {}
+	virtual void End() {}
+
+	virtual ~FEvaluationNotifyInstance() {}
+
+	TObjectPtr<UAnimNotifyState> AnimNotify = nullptr;
+	float StartTime = 0.f;
+	float EndTime = 0.f;
+	bool bActive = false;
+};
+
+// add procedural delta to the root motion attribute 
+USTRUCT(BlueprintInternalUseOnly)
+struct EVALUATIONNOTIFIESRUNTIME_API FAnimNode_EvaluationNotifies : public FAnimNode_SkeletalControlBase
+{
+	GENERATED_BODY()
+
+	// Animation Asset for incorporating root motion data. If CurrentAnimAsset is set, and the animation has root motion rotation within the TargetTime, then those rotations will be scaled to reach the TargetOrientation
+	UPROPERTY(EditAnywhere, Transient, BlueprintReadWrite, Category=Evaluation, meta=(PinShownByDefault))
+	TObjectPtr<UAnimationAsset> CurrentAnimAsset;
+	
+	// Current playback time in seconds of the CurrentAnimAsset
+	UPROPERTY(EditAnywhere, Transient, BlueprintReadWrite, Category=Evaluation, meta=(PinShownByDefault))
+	float CurrentAnimAssetTime = 0.f;
+	
+	// Current playback time in seconds of the CurrentAnimAsset
+	UPROPERTY(EditAnywhere, Transient, BlueprintReadWrite, Category=Evaluation, meta=(PinShownByDefault))
+	TMap<FName, FTransform> NamedTransforms;
+
+	// FAnimNodeBase interface
+	virtual void Initialize_AnyThread(const FAnimationInitializeContext& Context) override;
+	// End of FAnimNodeBase interface
+	
+	// FAnimNode_SkeletalControlBase interface
+	virtual void UpdateInternal(const FAnimationUpdateContext& Context) override;
+	virtual void EvaluateSkeletalControl_AnyThread(FComponentSpacePoseContext& Output, TArray<FBoneTransform>& OutBoneTransforms) override;
+	virtual bool IsValidToEvaluate(const USkeleton* Skeleton, const FBoneContainer& RequiredBones) override { return true; }
+	// End of FAnimNode_SkeletalControlBase interface
+
+	static void RegisterEvaluationHandler(UClass* NotifyType, UScriptStruct* Handler);
+	static void UnregisterEvaluationHandler(UClass* NotifyType);
+
+private:
+	void GetEvaluationNotifiesFromAnimation(const UAnimSequenceBase* Animation, TArray<FInstancedStruct>& OutNotifyInstances);
+
+	FTransform RootBoneTransform;
+	TArray<FInstancedStruct> Tags;
+	float PreviousAnimAssetTime = 0;
+	UPROPERTY(Transient)
+	TObjectPtr<const UAnimSequenceBase> CurrentSequence = nullptr;
+
+	static TMap<UClass*, UScriptStruct*> NotifyEvaluationHandlerMap;
+};
diff --git a/Engine/Plugins/Experimental/EvaluationNotifies/Source/Runtime/Public/EvaluationNotifies/AnimNotifyState_Alignment.h b/Engine/Plugins/Experimental/EvaluationNotifies/Source/Runtime/Public/EvaluationNotifies/AnimNotifyState_Alignment.h
new file mode 100644
--- /dev/null
+++ b/Engine/Plugins/Experimental/EvaluationNotifies/Source/Runtime/Public/EvaluationNotifies/AnimNotifyState_Alignment.h
@@ -0,0 +1,35 @@
+// Copyright Epic Games, Inc. All Rights Reserved.
+
+#pragma once
+
+#include "AnimNode_EvaluationNotifies.h"
+#include "AnimNotifyState_Alignment.generated.h"
+
+UCLASS(BlueprintType, DisplayName="Alignment")
+class UNotifyState_Alignment : public UAnimNotifyState
+{
+	GENERATED_BODY()
+	
+	public:
+
+	UPROPERTY(EditAnywhere, Category = Properties)
+	FBoneReference AlignBone;
+	
+	UPROPERTY(EditAnywhere, Category = Properties)
+	FTransform AlignOffset;
+
+	UPROPERTY(EditAnywhere, Category = Properties)
+	FName TransformName;
+};
+
+USTRUCT()
+struct FAlignmentNotifyInstance : public FEvaluationNotifyInstance
+{
+	GENERATED_BODY()
+
+	virtual void Start(const UAnimSequenceBase* AnimationAsset) override;
+	virtual void Update(const UAnimSequenceBase* AnimationAsset, float CurrentTime, float DeltaTime,
+		FTransform& RootBoneTransform, const TMap<FName, FTransform>& NamedTransforms, FComponentSpacePoseContext& Output, TArray<FBoneTransform>& OutBoneTransforms) override;
+
+	FBoneReference AlignBone;
+};
\ No newline at end of file
diff --git a/Engine/Plugins/Experimental/EvaluationNotifies/Source/Runtime/Public/EvaluationNotifies/AnimNotifyState_TwoBoneIK.h b/Engine/Plugins/Experimental/EvaluationNotifies/Source/Runtime/Public/EvaluationNotifies/AnimNotifyState_TwoBoneIK.h
new file mode 100644
--- /dev/null
+++ b/Engine/Plugins/Experimental/EvaluationNotifies/Source/Runtime/Public/EvaluationNotifies/AnimNotifyState_TwoBoneIK.h
@@ -0,0 +1,112 @@
+// Copyright Epic Games, Inc. All Rights Reserved.
+
+#pragma once
+
+#include "AnimNode_EvaluationNotifies.h"
+#include "CommonAnimTypes.h"
+#include "AnimNotifyState_TwoBoneIK.generated.h"
+
+UCLASS(BlueprintType, DisplayName="Two Bone IK")
+class UNotifyState_TwoBoneIK : public UAnimNotifyState
+{
+	GENERATED_BODY()
+	
+	public:
+	
+	/** Name of bone to control. This is the main bone chain to modify from. **/
+	UPROPERTY(EditAnywhere, Category=IK)
+	FBoneReference IKBone;
+	
+	/** Name of bone to IK Relative to - IK will target a position that has the same offset from the EffectorLocation, as the IKBone has from this Bone in the source pose. **/
+    UPROPERTY(EditAnywhere, Category=IK)
+    FBoneReference RelativeToBone;
+
+	/** Limits to use if stretching is allowed. This value determines when to start stretch. For example, 0.9 means once it reaches 90% of the whole length of the limb, it will start apply. */
+	UPROPERTY(EditAnywhere, Category=IK, meta = (editcondition = "bAllowStretching", ClampMin = "0.0", UIMin = "0.0"))
+	double StartStretchRatio = 1.0;
+
+	/** Limits to use if stretching is allowed. This value determins what is the max stretch scale. For example, 1.5 means it will stretch until 150 % of the whole length of the limb.*/
+	UPROPERTY(EditAnywhere, Category= IK, meta = (editcondition = "bAllowStretching", ClampMin = "0.0", UIMin = "0.0"))
+	double MaxStretchScale = 1.2;
+	
+	/** Limits to use if stretching is allowed. This value determins what is the max stretch scale. For example, 1.5 means it will stretch until 150 % of the whole length of the limb.*/
+	UPROPERTY(EditAnywhere, Category= IK)
+	float BlendInTime = 0.1f;
+	UPROPERTY(EditAnywhere, Category= IK)
+	float BlendOutTime = 0.1f;
+
+	/** Effector Location. Target Location to reach. */
+	UPROPERTY(EditAnywhere, BlueprintReadWrite, Category = Effector, meta = (PinShownByDefault))
+	FVector EffectorLocation = FVector::ZeroVector;
+
+	UPROPERTY(EditAnywhere, BlueprintReadWrite, Category = Effector, meta = (PinShownByDefault))
+	FName EffectorLocationTransformName;
+
+	UPROPERTY(EditAnywhere, Category=Effector)
+	FBoneSocketTarget EffectorTarget;
+
+	/** Joint Target Location. Location used to orient Joint bone. **/
+	UPROPERTY(EditAnywhere, BlueprintReadWrite, Category=JointTarget, meta=(PinShownByDefault))
+	FVector JointTargetLocation = FVector::ZeroVector;
+
+	UPROPERTY(EditAnywhere, Category = JointTarget)
+	FBoneSocketTarget JointTarget;
+
+	/** Specify which axis it's aligned. Used when removing twist */
+	UPROPERTY(EditAnywhere, Category = IK, meta = (editcondition = "!bAllowTwist"))
+	FAxis TwistAxis;
+
+	/** Reference frame of Effector Location. */
+	UPROPERTY(EditAnywhere, Category=Effector)
+	TEnumAsByte<enum EBoneControlSpace> EffectorLocationSpace = BCS_ComponentSpace;
+
+	/** Reference frame of Joint Target Location. */
+	UPROPERTY(EditAnywhere, Category=JointTarget)
+	TEnumAsByte<enum EBoneControlSpace> JointTargetLocationSpace = BCS_ComponentSpace;
+
+	/** Should stretching be allowed, to be prevent over extension */
+	UPROPERTY(EditAnywhere, Category=IK)
+	uint8 bAllowStretching:1 = false;
+
+	/** Set end bone to use End Effector rotation */
+	UPROPERTY(EditAnywhere, Category=IK)
+	uint8 bTakeRotationFromEffectorSpace : 1 = false;
+
+	/** Keep local rotation of end bone */
+	UPROPERTY(EditAnywhere, Category = IK)
+	uint8 bMaintainEffectorRelRot : 1 = false;
+
+	/** Whether or not to apply twist on the chain of joints. This clears the twist value along the TwistAxis */
+	UPROPERTY(EditAnywhere, Category = IK)
+	uint8 bAllowTwist : 1 = true;
+};
+
+USTRUCT()
+struct FTwoBoneIKNotifyInstance : public FEvaluationNotifyInstance
+{
+	FTwoBoneIKNotifyInstance()
+    	:  CachedUpperLimbIndex(INDEX_NONE)
+    	, CachedLowerLimbIndex(INDEX_NONE)
+    {
+    }
+
+	
+	GENERATED_BODY()
+
+	virtual void Start(const UAnimSequenceBase* AnimationAsset) override;
+	virtual void Update(const UAnimSequenceBase* AnimationAsset, float CurrentTime, float DeltaTime,
+		FTransform& RootBoneTransform, const TMap<FName, FTransform>& NamedTransforms, FComponentSpacePoseContext& Output, TArray<FBoneTransform>& OutBoneTransforms) override;
+
+	FBoneReference IKBone;
+	FBoneReference RelativeToBone;
+	FBoneSocketTarget EffectorTarget;
+	
+	UPROPERTY(EditAnywhere, Category = JointTarget)
+	FBoneSocketTarget JointTarget;
+	
+	// cached limb index for upper
+	FCompactPoseBoneIndex CachedUpperLimbIndex;
+	
+	// cached limb index for lower
+	FCompactPoseBoneIndex CachedLowerLimbIndex;
+};
\ No newline at end of file
diff --git a/Engine/Plugins/Experimental/EvaluationNotifies/Source/Runtime/Public/EvaluationNotifiesRuntimeModule.h b/Engine/Plugins/Experimental/EvaluationNotifies/Source/Runtime/Public/EvaluationNotifiesRuntimeModule.h
new file mode 100644
--- /dev/null
+++ b/Engine/Plugins/Experimental/EvaluationNotifies/Source/Runtime/Public/EvaluationNotifiesRuntimeModule.h
@@ -0,0 +1,9 @@
+// Copyright Epic Games, Inc. All Rights Reserved.
+
+#pragma once
+
+#include "Modules/ModuleInterface.h"
+
+class IEvaluationNotifiesRuntimeModule : public IModuleInterface
+{
+};
