Updates to Alignment notify to support blending into alignment with different curve types, and using future root motion information to only apply the necessary delta when an alignment bone isn't available

diff --git a/Engine/Plugins/Experimental/EvaluationNotifies/Source/Runtime/EvaluationNotifiesRuntime.Build.cs b/Engine/Plugins/Experimental/EvaluationNotifies/Source/Runtime/EvaluationNotifiesRuntime.Build.cs
--- a/Engine/Plugins/Experimental/EvaluationNotifies/Source/Runtime/EvaluationNotifiesRuntime.Build.cs
+++ b/Engine/Plugins/Experimental/EvaluationNotifies/Source/Runtime/EvaluationNotifiesRuntime.Build.cs
@@ -1,11 +1,15 @@
 // Copyright Epic Games, Inc. All Rights Reserved.
 
+using System.IO;
+
 namespace UnrealBuildTool.Rules
 {
 	public class EvaluationNotifiesRuntime : ModuleRules
 	{
 		public EvaluationNotifiesRuntime(ReadOnlyTargetRules Target) : base(Target)
 		{
+			PrivateIncludePaths.Add(Path.Combine(EngineDirectory,"Source/ThirdParty/AHEasing/AHEasing-1.3.2"));
+			
 			PublicDependencyModuleNames.AddRange(
 			new string[]
 			{
diff --git a/Engine/Plugins/Experimental/EvaluationNotifies/Source/Runtime/Private/EvaluationNotifies/AnimNotifyState_Alignment.cpp b/Engine/Plugins/Experimental/EvaluationNotifies/Source/Runtime/Private/EvaluationNotifies/AnimNotifyState_Alignment.cpp
--- a/Engine/Plugins/Experimental/EvaluationNotifies/Source/Runtime/Private/EvaluationNotifies/AnimNotifyState_Alignment.cpp
+++ b/Engine/Plugins/Experimental/EvaluationNotifies/Source/Runtime/Private/EvaluationNotifies/AnimNotifyState_Alignment.cpp
@@ -3,6 +3,7 @@
 #include "EvaluationNotifies/AnimNotifyState_Alignment.h"
 #include "Animation/AnimInstanceProxy.h"
 #include "Animation/AnimRootMotionProvider.h"
+#include "AHEasing/easing.h"
 
 #include UE_INLINE_GENERATED_CPP_BY_NAME(AnimNotifyState_Alignment)
 
@@ -10,8 +11,37 @@ void FAlignmentNotifyInstance::Start(const UAnimSequenceBase* AnimationAsset)
 {
 	UNotifyState_Alignment* AlignmentNotify = Cast<UNotifyState_Alignment>(AnimNotify);
 	AlignBone = AlignmentNotify->AlignBone;
+	PredictionMode = AlignmentNotify->PredictionMode;
+	PostNotifyTargetTime = AlignmentNotify->PostNotifyTargetTime;
+	WeightCurve = AlignmentNotify->WeightCurve;
 }
 
+float FAlignmentNotifyInstance::GetWeightThisFrame(float CurrentTime, float DeltaTime)
+{
+	float Weight;
+	const float CurrentRelativeTime = (CurrentTime-StartTime)/(EndTime-StartTime);
+	switch(WeightCurve)
+	{
+		case EAlignmentWeightCurveType::EaseIn:
+			Weight = CubicEaseIn(CurrentRelativeTime);
+			break;
+		case EAlignmentWeightCurveType::EaseOut:
+			Weight = CubicEaseOut(CurrentRelativeTime);
+			break;
+		case EAlignmentWeightCurveType::EaseInOut:
+			Weight = CubicEaseInOut(CurrentRelativeTime);
+			break;
+	
+		default:
+			Weight = CurrentRelativeTime;
+	}
+	const float RemainingWeight = 1 - PreviousCurveWeight;
+	const float DeltaWeight = Weight - PreviousCurveWeight;
+	const float ThisFrameWeight = DeltaWeight / RemainingWeight;
+	PreviousCurveWeight = Weight;
+	return ThisFrameWeight;
+}	
+
 void FAlignmentNotifyInstance::Update(const UAnimSequenceBase* AnimationAsset, float CurrentTime, float DeltaTime, 
                                    FTransform& RootBoneTransform, const TMap<FName, FTransform>& NamedTransforms, FComponentSpacePoseContext& Output, TArray<FBoneTransform>& OutBoneTransforms)
 {
@@ -32,18 +62,53 @@ void FAlignmentNotifyInstance::Update(const UAnimSequenceBase* AnimationAsset, f
     	{
     		if (const FTransform* TargetTransform = NamedTransforms.Find(AlignmentNotify->TransformName))
     		{
-    			// simple linear blend weight
-    			float TimeRemaining = EndTime-CurrentTime;
-    			float Weight = TimeRemaining <= 0.f ? 1.f : FMath::Min(1.f, DeltaTime / TimeRemaining);
-    		
-				const FBoneContainer& RequiredBones = Output.AnimInstanceProxy->GetRequiredBones();	
-    			FCompactPoseBoneIndex AlignBoneIndex = AlignBone.GetCompactPoseIndex(RequiredBones);
-    			if (AlignBoneIndex.IsValid())
+    			FTransform PredictedRootMotion = FTransform::Identity;
+    			bool bPredictedRootMotionValid = false;
+
+    			if (PredictionMode == EAlignmentPredictionMode::AlignBone)
     			{
-    				const FTransform& AlignBoneTransform  = Output.Pose.GetComponentSpaceTransform(AlignBoneIndex);
-    				RootBoneTransform = AlignBoneTransform * ThisFrameRootMotionTransform * RootBoneTransform;;
+    				const FBoneContainer& RequiredBones = Output.AnimInstanceProxy->GetRequiredBones();	
+    				FCompactPoseBoneIndex AlignBoneIndex = AlignBone.GetCompactPoseIndex(RequiredBones);
+    				if (AlignBoneIndex.IsValid())
+    				{
+    					const FTransform& AlignBoneTransform  = Output.Pose.GetComponentSpaceTransform(AlignBoneIndex);
+    					RootBoneTransform = AlignBoneTransform * ThisFrameRootMotionTransform * RootBoneTransform;;
+    				}
     			}
+    			else if (PredictionMode == EAlignmentPredictionMode::RootMotion)
+    			{
+    				const float PredictionDelta = EndTime + PostNotifyTargetTime - CurrentTime ;
+    				PredictedRootMotion = AnimationAsset->ExtractRootMotion(CurrentTime, PredictionDelta, false);
+    				RootBoneTransform = RootBoneTransform * PredictedRootMotion;
 
+    				// this predicted transform can be reused for the weight curve, but only if there was no time offset
+    				bPredictedRootMotionValid = (PostNotifyTargetTime == 0.0f);
+    			}
+ 
+    			float Weight = GetWeightThisFrame(CurrentTime, DeltaTime);
+    		
+				if (WeightCurve == EAlignmentWeightCurveType::FromRootMotion)
+				{
+					// override the weight with a weight computed based on remaining root motion in the animation
+					
+					if (Weight < 1.0f)
+					{
+						if (!bPredictedRootMotionValid)
+						{
+							// if we didn't already predict root motion above, or we predicted to a time offset, then compute the root motion to the end of the notify here
+							const float PredictionDelta = EndTime - CurrentTime;
+							PredictedRootMotion = AnimationAsset->ExtractRootMotion(CurrentTime, PredictionDelta, false);
+						}
+						
+						const float RootMotionDistanceRemaining = PredictedRootMotion.GetLocation().Length();
+						if (RootMotionDistanceRemaining > UE_SMALL_NUMBER)
+						{
+							const float RootMotionDistanceThisFrame = ThisFrameRootMotionTransform.GetLocation().Length();
+							Weight = RootMotionDistanceThisFrame / RootMotionDistanceRemaining;
+						}
+					}
+				}
+    			
     			RootBoneTransform = AlignmentNotify->AlignOffset * RootBoneTransform;
     		
     			FTransform RelativeOffset = TargetTransform->GetRelativeTransform(RootBoneTransform);
diff --git a/Engine/Plugins/Experimental/EvaluationNotifies/Source/Runtime/Public/EvaluationNotifies/AnimNotifyState_Alignment.h b/Engine/Plugins/Experimental/EvaluationNotifies/Source/Runtime/Public/EvaluationNotifies/AnimNotifyState_Alignment.h
--- a/Engine/Plugins/Experimental/EvaluationNotifies/Source/Runtime/Public/EvaluationNotifies/AnimNotifyState_Alignment.h
+++ b/Engine/Plugins/Experimental/EvaluationNotifies/Source/Runtime/Public/EvaluationNotifies/AnimNotifyState_Alignment.h
@@ -5,21 +5,59 @@
 #include "AnimNode_EvaluationNotifies.h"
 #include "AnimNotifyState_Alignment.generated.h"
 
+UENUM()
+enum class EAlignmentWeightCurveType
+{
+	FromRootMotion,
+	Linear,
+	EaseIn,
+	EaseOut,
+	EaseInOut
+};
+
+UENUM()
+enum class EAlignmentPredictionMode
+{
+	// Root Motion data from the AnimSequence will be used to predict the animated target transform, 
+	RootMotion,
+	// A Bone in the AnimSequence has been animated to represent the target transform
+	AlignBone,
+	// Warping will be applied every frame without considering the future movement in the AnimSequence
+	None,
+};
+
 UCLASS(BlueprintType, DisplayName="Alignment")
 class UNotifyState_Alignment : public UAnimNotifyState
 {
 	GENERATED_BODY()
 	
 	public:
-
+	
+	// Curve used to control the rate at which we blend into alignment
 	UPROPERTY(EditAnywhere, Category = Properties)
-	FBoneReference AlignBone;
+	EAlignmentWeightCurveType WeightCurve = EAlignmentWeightCurveType::FromRootMotion;
 	
+	// Named transform to align to
+    UPROPERTY(EditAnywhere, Category = Properties)
+    FName TransformName;
+
+	// Offset from Root (or AlignBone) to Align to the target transform
 	UPROPERTY(EditAnywhere, Category = Properties)
 	FTransform AlignOffset;
 
+	// Method for predicting how the animation will move the character relative to the alignment target
 	UPROPERTY(EditAnywhere, Category = Properties)
-	FName TransformName;
+	EAlignmentPredictionMode PredictionMode;
+
+	// Time in seconds after the end of the NotifyState when the root should be in alignment with the target (warping will happen over the notify duration, non-warped root motion will complete the alignment)
+	UPROPERTY(EditAnywhere, Category = Properties, meta = (EditCondition = "PredictionMode == EAlignmentPredictionMode::RootMotion"))
+	float PostNotifyTargetTime = 0.f;
+
+	// Bone to Align to target transform
+	UPROPERTY(EditAnywhere, Category = Properties, meta = (EditCondition = "PredictionMode == EAlignmentPredictionMode::AlignBone"))
+	FBoneReference AlignBone;
+
+	
 };
 
 USTRUCT()
@@ -30,6 +68,13 @@ struct FAlignmentNotifyInstance : public FEvaluationNotifyInstance
 	virtual void Start(const UAnimSequenceBase* AnimationAsset) override;
 	virtual void Update(const UAnimSequenceBase* AnimationAsset, float CurrentTime, float DeltaTime,
 		FTransform& RootBoneTransform, const TMap<FName, FTransform>& NamedTransforms, FComponentSpacePoseContext& Output, TArray<FBoneTransform>& OutBoneTransforms) override;
+	
+	
+	float GetWeightThisFrame(float CurrentTime, float DeltaTime);
 
+	float PreviousCurveWeight = 0.0;
 	FBoneReference AlignBone;
+	float PostNotifyTargetTime = 0.f;
+	EAlignmentPredictionMode PredictionMode;
+	EAlignmentWeightCurveType WeightCurve;
 };
\ No newline at end of file
