Alignment visual logging and math fixes

diff --git a/Engine/Plugins/Experimental/EvaluationNotifies/Source/Runtime/Private/EvaluationNotifies/AnimNotifyState_Alignment.cpp b/Engine/Plugins/Experimental/EvaluationNotifies/Source/Runtime/Private/EvaluationNotifies/AnimNotifyState_Alignment.cpp
--- a/Engine/Plugins/Experimental/EvaluationNotifies/Source/Runtime/Private/EvaluationNotifies/AnimNotifyState_Alignment.cpp
+++ b/Engine/Plugins/Experimental/EvaluationNotifies/Source/Runtime/Private/EvaluationNotifies/AnimNotifyState_Alignment.cpp
@@ -4,9 +4,16 @@
 #include "Animation/AnimInstanceProxy.h"
 #include "Animation/AnimRootMotionProvider.h"
 #include "AHEasing/easing.h"
+#include "VisualLogger/VisualLogger.h"
 
 #include UE_INLINE_GENERATED_CPP_BY_NAME(AnimNotifyState_Alignment)
 
+bool g_bEnableExtraAlignmentVLogs = false;
+static FAutoConsoleVariableRef CVarEnableExtraAlignmentVLogs(
+	TEXT("Animation.Notify.Alignment.EnableExtraVLogs"),
+	g_bEnableExtraAlignmentVLogs,
+	TEXT("Enable Extra Alignment Visual Logger drawing (with a performance overhead)"));
+
 void FAlignmentNotifyInstance::Start(const UAnimSequenceBase* AnimationAsset)
 {
 	UNotifyState_Alignment* AlignmentNotify = Cast<UNotifyState_Alignment>(AnimNotify);
@@ -54,6 +61,7 @@ void FAlignmentNotifyInstance::Update(const UAnimSequenceBase* AnimationAsset, f
 
 	const UE::Anim::IAnimRootMotionProvider* RootMotionProvider = UE::Anim::IAnimRootMotionProvider::Get();
     ensureMsgf(RootMotionProvider, TEXT("Alignment expected a valid root motion delta provider interface."));
+	FTransform AlignmentTransform = RootBoneTransform;
 
     if (RootMotionProvider)
     {
@@ -62,6 +70,9 @@ void FAlignmentNotifyInstance::Update(const UAnimSequenceBase* AnimationAsset, f
     	{
     		if (const FTransform* TargetTransform = NamedTransforms.Find(AlignmentNotify->TransformName))
     		{
+				static const FBox UnitBox(FVector(-10, -10, -10), FVector(10, 10, 10));
+				UE_VLOG_OBOX(Output.AnimInstanceProxy->GetAnimInstanceObject(), "Alignment", Display, UnitBox, TargetTransform->ToMatrixWithScale(), FColor::Green, TEXT(""));
+    		
     			FTransform PredictedRootMotion = FTransform::Identity;
     			bool bPredictedRootMotionValid = false;
 
@@ -72,19 +83,42 @@ void FAlignmentNotifyInstance::Update(const UAnimSequenceBase* AnimationAsset, f
     				if (AlignBoneIndex.IsValid())
     				{
     					const FTransform& AlignBoneTransform  = Output.Pose.GetComponentSpaceTransform(AlignBoneIndex);
-    					RootBoneTransform = AlignBoneTransform * ThisFrameRootMotionTransform * RootBoneTransform;;
+    					AlignmentTransform = AlignBoneTransform * ThisFrameRootMotionTransform * AlignmentTransform;;
     				}
     			}
     			else if (PredictionMode == EAlignmentPredictionMode::RootMotion)
     			{
-    				const float PredictionDelta = EndTime + PostNotifyTargetTime - CurrentTime ;
-    				PredictedRootMotion = AnimationAsset->ExtractRootMotion(CurrentTime, PredictionDelta, false);
-    				RootBoneTransform = RootBoneTransform * PredictedRootMotion;
+#if ENABLE_VISUAL_LOG
+    				if (g_bEnableExtraAlignmentVLogs)
+    				{
+    					// render the trajectory of the animation until the alignment point
+    					
+    					FTransform PredictedTransform = RootBoneTransform;
+    					float PredictionTime = CurrentTime;
+    					float TargetTime = EndTime + PostNotifyTargetTime;
+    					float TimeStep = 1.0f/30.0f;
+    					bool AltColor = false;
+    					while(PredictionTime < TargetTime)
+    					{
+    						FVector PrevPosition = PredictedTransform.GetLocation();
+    						PredictedTransform = AnimationAsset->ExtractRootMotion(PredictionTime, TimeStep, false) * PredictedTransform;
+    						UE_VLOG_SEGMENT(Output.AnimInstanceProxy->GetAnimInstanceObject(), "Alignment", Display, PrevPosition, PredictedTransform.GetLocation(), AltColor ? FColor::Yellow : FColor::Red, TEXT(""));
+    						AltColor = !AltColor;
+    						PredictionTime += TimeStep;
+    					}
+    				}
+#endif
+    				
+					const float PredictionDelta = EndTime + PostNotifyTargetTime - CurrentTime ;
+					PredictedRootMotion = AnimationAsset->ExtractRootMotion(CurrentTime, PredictionDelta, false);
+					AlignmentTransform = PredictedRootMotion * AlignmentTransform;
 
     				// this predicted transform can be reused for the weight curve, but only if there was no time offset
     				bPredictedRootMotionValid = (PostNotifyTargetTime == 0.0f);
+    				
+					UE_VLOG_OBOX(Output.AnimInstanceProxy->GetAnimInstanceObject(), "Alignment", Display, UnitBox, AlignmentTransform.ToMatrixWithScale(), FColor::Blue, TEXT(""));
     			}
- 
+   			
     			float Weight = GetWeightThisFrame(CurrentTime, DeltaTime);
     		
 				if (WeightCurve == EAlignmentWeightCurveType::FromRootMotion)
@@ -109,14 +143,19 @@ void FAlignmentNotifyInstance::Update(const UAnimSequenceBase* AnimationAsset, f
 					}
 				}
     			
-    			RootBoneTransform = AlignmentNotify->AlignOffset * RootBoneTransform;
+    			AlignmentTransform = AlignmentNotify->AlignOffset * AlignmentTransform;
     		
-    			FTransform RelativeOffset = TargetTransform->GetRelativeTransform(RootBoneTransform);
+    			FQuat CurrentRotation = ThisFrameRootMotionTransform.GetRotation();
+    			CurrentRotation = CurrentRotation * FMath::Lerp(FQuat::Identity, AlignmentTransform.GetRotation().Inverse() * TargetTransform->GetRotation(), Weight);
+    			
+    			FVector CurrentTranslationWS = RootBoneTransform.TransformVector(ThisFrameRootMotionTransform.GetTranslation());
+    			CurrentTranslationWS += (TargetTransform->GetLocation() - AlignmentTransform.GetLocation())* Weight;
 
-    			FTransform ScaledOffset = FTransform(FMath::Lerp(FQuat::Identity, RelativeOffset.GetRotation(), Weight),
-														Weight * RelativeOffset.GetTranslation());
+    			ThisFrameRootMotionTransform.SetRotation(CurrentRotation);
+    			ThisFrameRootMotionTransform.SetTranslation(RootBoneTransform.InverseTransformVector(CurrentTranslationWS));
 
-    			RootMotionProvider->OverrideRootMotion(ThisFrameRootMotionTransform * ScaledOffset, Output.CustomAttributes);
+    			
+    			RootMotionProvider->OverrideRootMotion(ThisFrameRootMotionTransform, Output.CustomAttributes);
     		}
     	}
     }
