In graph animation warping/alignment improvements

diff --git a/Engine/Plugins/Experimental/EvaluationNotifies/Source/Runtime/Private/EvaluationNotifies/AnimNode_EvaluationNotifies.cpp b/Engine/Plugins/Experimental/EvaluationNotifies/Source/Runtime/Private/EvaluationNotifies/AnimNode_EvaluationNotifies.cpp
--- a/Engine/Plugins/Experimental/EvaluationNotifies/Source/Runtime/Private/EvaluationNotifies/AnimNode_EvaluationNotifies.cpp
+++ b/Engine/Plugins/Experimental/EvaluationNotifies/Source/Runtime/Private/EvaluationNotifies/AnimNode_EvaluationNotifies.cpp
@@ -110,7 +110,7 @@ void FAnimNode_EvaluationNotifies::EvaluateSkeletalControl_AnyThread(FComponentS
 					Tag.Start(CurrentSequence);
 				}
 
-				Tag.Update(CurrentSequence, CurrentAnimAssetTime, DeltaSeconds, RootBoneTransform, NamedTransforms, Output, OutBoneTransforms);
+				Tag.Update(CurrentSequence, CurrentAnimAssetTime, DeltaSeconds, CurrentAnimAssetMirrored, MirrorDataTable, RootBoneTransform, NamedTransforms, Output, OutBoneTransforms);
 			}
 			else
 			{
diff --git a/Engine/Plugins/Experimental/EvaluationNotifies/Source/Runtime/Private/EvaluationNotifies/AnimNotifyState_Alignment.cpp b/Engine/Plugins/Experimental/EvaluationNotifies/Source/Runtime/Private/EvaluationNotifies/AnimNotifyState_Alignment.cpp
--- a/Engine/Plugins/Experimental/EvaluationNotifies/Source/Runtime/Private/EvaluationNotifies/AnimNotifyState_Alignment.cpp
+++ b/Engine/Plugins/Experimental/EvaluationNotifies/Source/Runtime/Private/EvaluationNotifies/AnimNotifyState_Alignment.cpp
@@ -4,161 +4,330 @@
 #include "Animation/AnimInstanceProxy.h"
 #include "Animation/AnimRootMotionProvider.h"
 #include "AHEasing/easing.h"
-#include "HAL/IConsoleManager.h"
+#include "Animation/AnimSequenceHelpers.h"
 #include "VisualLogger/VisualLogger.h"
 
 #include UE_INLINE_GENERATED_CPP_BY_NAME(AnimNotifyState_Alignment)
 
-bool g_bEnableExtraAlignmentVLogs = false;
-static FAutoConsoleVariableRef CVarEnableExtraAlignmentVLogs(
-	TEXT("Animation.Notify.Alignment.EnableExtraVLogs"),
-	g_bEnableExtraAlignmentVLogs,
-	TEXT("Enable Extra Alignment Visual Logger drawing (with a performance overhead)"));
+namespace
+{
+	constexpr float FrameTime = 1.0f/30.0f;
+
+	FTransform ExtractRootMotionHelper(const UAnimSequenceBase* AnimationAsset, const UMirrorDataTable* MirrorDataTable, bool bIsMirrored, float CurrentTime, float DeltaTime, bool bLoop)
+	{
+		if (bIsMirrored)
+		{
+			return UE::Anim::ExtractRootMotionFromAnimationAsset(AnimationAsset, MirrorDataTable, CurrentTime, DeltaTime, bLoop);
+		}
+		else
+		{
+			const FAnimExtractContext Context(static_cast<double>(CurrentTime), true, FDeltaTimeRecord(DeltaTime), bLoop);
+			return AnimationAsset->ExtractRootMotion(Context);
+		}
+	}
+}
 
 void FAlignmentNotifyInstance::Start(const UAnimSequenceBase* AnimationAsset)
 {
 	UNotifyState_Alignment* AlignmentNotify = Cast<UNotifyState_Alignment>(AnimNotify);
 	AlignBone = AlignmentNotify->AlignBone;
-	PredictionMode = AlignmentNotify->PredictionMode;
-	PostNotifyTargetTime = AlignmentNotify->PostNotifyTargetTime;
-	WeightCurve = AlignmentNotify->WeightCurve;
+	bFirstFrame = true;
 }
 
-float FAlignmentNotifyInstance::GetWeightThisFrame(float CurrentTime, float DeltaTime)
+float SampleCurve(float SampleTime, float StartTime, float EndTime, const TArray<float>& CurveData)
 {
-	float Weight;
-	const float CurrentRelativeTime = (CurrentTime-StartTime)/(EndTime-StartTime);
-	switch(WeightCurve)
+	if (SampleTime <= StartTime)
+	{
+		return 0.0f;
+	}
+	if (SampleTime >= EndTime)
 	{
-		case EAlignmentWeightCurveType::EaseIn:
-			Weight = CubicEaseIn(CurrentRelativeTime);
-			break;
-		case EAlignmentWeightCurveType::EaseOut:
-			Weight = CubicEaseOut(CurrentRelativeTime);
-			break;
-		case EAlignmentWeightCurveType::EaseInOut:
-			Weight = CubicEaseInOut(CurrentRelativeTime);
-			break;
+		return 1.0;
+	}
+	
+	// do nearest neighbor sampling for simplicity and performance
+	const int StartFrame = FMath::Clamp(round(StartTime / FrameTime), 0, CurveData.Num()-1);
+	const int EndFrame = FMath::Clamp(round(EndTime / FrameTime), 0, CurveData.Num()-1);
+	const int SampleFrame = FMath::Clamp(round(SampleTime / FrameTime), 0, CurveData.Num()-1);
+
+	return (CurveData[SampleFrame] - CurveData[StartFrame])/(CurveData[EndFrame] - CurveData[StartFrame]);
+}
+
+
+
+
+float FAlignmentNotifyInstance::GetWeight(float CurrentTime, const FAlignmentWarpCurve& WarpCurve) const 
+{
+	UNotifyState_Alignment* AlignmentNotify = Cast<UNotifyState_Alignment>(AnimNotify);
+	float Weight;
 	
-		default:
-			Weight = CurrentRelativeTime;
+
+	if (WarpCurve.CurveType == EAlignmentWeightCurveType::FromRootMotionTranslation && !AnimTrajectoryData.TranslationCurve.IsEmpty())
+	{
+		const float Duration = (EndTime - StartTime);
+		return SampleCurve(CurrentTime - StartTime, WarpCurve.StartRatio * Duration, WarpCurve.EndRatio * Duration, AnimTrajectoryData.TranslationCurve);
+	}
+	else if (WarpCurve.CurveType == EAlignmentWeightCurveType::FromRootMotionRotation && !AnimTrajectoryData.RotationCurve.IsEmpty())
+	{
+		const float Duration = (EndTime - StartTime);
+		return SampleCurve(CurrentTime - StartTime, WarpCurve.StartRatio * Duration, WarpCurve.EndRatio * Duration, AnimTrajectoryData.RotationCurve);
+	}
+	else
+	{
+		const float StartTimeWithRatio = FMath::Lerp(StartTime, EndTime, WarpCurve.StartRatio);
+		const float EndTimeWithRatio = FMath::Lerp(StartTime, EndTime, WarpCurve.EndRatio);
+			
+		const float CurrentRelativeTime = FMath::Clamp((CurrentTime-StartTimeWithRatio)/(EndTimeWithRatio-StartTimeWithRatio), 0,1);
+		
+		switch(WarpCurve.CurveType)
+		{
+			case EAlignmentWeightCurveType::EaseIn:
+				Weight = CubicEaseIn(CurrentRelativeTime);
+				break;
+			case EAlignmentWeightCurveType::EaseOut:
+				Weight = CubicEaseOut(CurrentRelativeTime);
+				break;
+			case EAlignmentWeightCurveType::EaseInOut:
+				Weight = CubicEaseInOut(CurrentRelativeTime);
+				break;
+			case EAlignmentWeightCurveType::Instant:
+				Weight = 1.f;
+				break;
+			case EAlignmentWeightCurveType::DoNotWarp:
+				Weight = 0.f;
+				break;
+			default:
+				Weight = CurrentRelativeTime;
+		}
 	}
-	const float RemainingWeight = 1 - PreviousCurveWeight;
-	const float DeltaWeight = Weight - PreviousCurveWeight;
-	const float ThisFrameWeight = DeltaWeight / RemainingWeight;
-	PreviousCurveWeight = Weight;
-	return ThisFrameWeight;
-}	
-
-void FAlignmentNotifyInstance::Update(const UAnimSequenceBase* AnimationAsset, float CurrentTime, float DeltaTime, 
+
+	
+	return Weight;
+}
+
+
+void FAlignmentNotifyInstance::Update(const UAnimSequenceBase* AnimationAsset, float CurrentTime, float DeltaTime, bool bIsMirrored, const UMirrorDataTable* MirrorDataTable,
                                    FTransform& RootBoneTransform, const TMap<FName, FTransform>& NamedTransforms, FComponentSpacePoseContext& Output, TArray<FBoneTransform>& OutBoneTransforms)
 {
 	UNotifyState_Alignment* AlignmentNotify = Cast<UNotifyState_Alignment>(AnimNotify);
-	if (!AlignBone.HasValidSetup())
+	
+	if (bFirstFrame)
 	{
+		bFirstFrame = false;
+		
 		const FBoneContainer& RequiredBones = Output.AnimInstanceProxy->GetRequiredBones();	
 		AlignBone.Initialize(RequiredBones);
-	}
+		
+		if (const FTransform* FoundTransform = NamedTransforms.Find(AlignmentNotify->TransformName))
+		{
+			TargetTransform = *FoundTransform;
+			
+			static const FBox UnitBox(FVector(-10, -10, -10), FVector(10, 10, 10));
+			UE_VLOG_OBOX(Output.AnimInstanceProxy->GetAnimInstanceObject(), "Alignment", Display, UnitBox, TargetTransform.ToMatrixWithScale(), FColor::Blue, TEXT(""));
+			
+			// get alignment transform
+			const float PredictionDelta = EndTime - CurrentTime ;
+			
 
-	const UE::Anim::IAnimRootMotionProvider* RootMotionProvider = UE::Anim::IAnimRootMotionProvider::Get();
-    ensureMsgf(RootMotionProvider, TEXT("Alignment expected a valid root motion delta provider interface."));
-	FTransform AlignmentTransform = RootBoneTransform;
-
-    if (RootMotionProvider)
-    {
-    	FTransform ThisFrameRootMotionTransform;
-    	if (RootMotionProvider->ExtractRootMotion(Output.CustomAttributes, ThisFrameRootMotionTransform))
-    	{
-    		if (const FTransform* TargetTransform = NamedTransforms.Find(AlignmentNotify->TransformName))
-    		{
-				static const FBox UnitBox(FVector(-10, -10, -10), FVector(10, 10, 10));
-				UE_VLOG_OBOX(Output.AnimInstanceProxy->GetAnimInstanceObject(), "Alignment", Display, UnitBox, TargetTransform->ToMatrixWithScale(), FColor::Green, TEXT(""));
-    		
-    			FTransform PredictedRootMotion = FTransform::Identity;
-    			bool bPredictedRootMotionValid = false;
-
-    			if (PredictionMode == EAlignmentPredictionMode::AlignBone)
-    			{
-    				const FBoneContainer& RequiredBones = Output.AnimInstanceProxy->GetRequiredBones();	
-    				FCompactPoseBoneIndex AlignBoneIndex = AlignBone.GetCompactPoseIndex(RequiredBones);
-    				if (AlignBoneIndex.IsValid())
-    				{
-    					const FTransform& AlignBoneTransform  = Output.Pose.GetComponentSpaceTransform(AlignBoneIndex);
-    					AlignmentTransform = AlignBoneTransform * ThisFrameRootMotionTransform * AlignmentTransform;;
-    				}
-    			}
-    			else if (PredictionMode == EAlignmentPredictionMode::RootMotion)
-    			{
-#if ENABLE_VISUAL_LOG
-    				if (g_bEnableExtraAlignmentVLogs)
-    				{
-    					// render the trajectory of the animation until the alignment point
-    					
-    					FTransform PredictedTransform = RootBoneTransform;
-    					float PredictionTime = CurrentTime;
-    					float TargetTime = EndTime + PostNotifyTargetTime;
-    					float TimeStep = 1.0f/30.0f;
-    					bool AltColor = false;
-    					while(PredictionTime < TargetTime)
-    					{
-    						FVector PrevPosition = PredictedTransform.GetLocation();
-    						PredictedTransform = AnimationAsset->ExtractRootMotion(FAnimExtractContext(static_cast<double>(PredictionTime), true, FDeltaTimeRecord(TimeStep), false)) * PredictedTransform;
-    						UE_VLOG_SEGMENT(Output.AnimInstanceProxy->GetAnimInstanceObject(), "Alignment", Display, PrevPosition, PredictedTransform.GetLocation(), AltColor ? FColor::Yellow : FColor::Red, TEXT(""));
-    						AltColor = !AltColor;
-    						PredictionTime += TimeStep;
-    					}
-    				}
-#endif
-    				
-					const float PredictionDelta = EndTime + PostNotifyTargetTime - CurrentTime ;
-					PredictedRootMotion = AnimationAsset->ExtractRootMotion(FAnimExtractContext(static_cast<double>(CurrentTime), true, FDeltaTimeRecord(PredictionDelta), false));
-					AlignmentTransform = PredictedRootMotion * AlignmentTransform;
-
-    				// this predicted transform can be reused for the weight curve, but only if there was no time offset
-    				bPredictedRootMotionValid = (PostNotifyTargetTime == 0.0f);
-    				
-					UE_VLOG_OBOX(Output.AnimInstanceProxy->GetAnimInstanceObject(), "Alignment", Display, UnitBox, AlignmentTransform.ToMatrixWithScale(), FColor::Blue, TEXT(""));
-    			}
-   			
-    			float Weight = GetWeightThisFrame(CurrentTime, DeltaTime);
-    		
-				if (WeightCurve == EAlignmentWeightCurveType::FromRootMotion)
+			if (AlignBone.HasValidSetup())
+			{
+				// get alignment bone relative to predicted root motion end point, and remove that as an offset to the alignment target
+				// Alignment bone is expected to be at the alignment end point at the beginning of the Notify window  (we don't have a good way to prediuct the component space transform of it in the future)
+				
+				FTransform PredictedRootMotion = ExtractRootMotionHelper(AnimationAsset, MirrorDataTable, bIsMirrored, CurrentTime, PredictionDelta, false);
+				FCompactPoseBoneIndex AlignBoneIndex = AlignBone.GetCompactPoseIndex(RequiredBones);
+				if (AlignBoneIndex.IsValid())
 				{
-					// override the weight with a weight computed based on remaining root motion in the animation
-					
-					if (Weight < 1.0f)
+					TargetTransform = Output.Pose.GetComponentSpaceTransform(AlignBoneIndex).Inverse() * PredictedRootMotion * TargetTransform;
+				}
+			}
+			
+			TargetTransform = AlignmentNotify->AlignOffset * TargetTransform;
+			StartingRootTransform = RootBoneTransform;
+			
+			// extract root motion trajectory todo: this should be cached and reused
+			
+			int NumFrames = (EndTime - CurrentTime) / FrameTime;
+
+			AnimTrajectoryData.Trajectory.SetNum(NumFrames);
+			AnimTrajectoryData.TranslationCurve.SetNum(NumFrames);
+			AnimTrajectoryData.RotationCurve.SetNum(NumFrames);
+			
+			float SteeringAngleThreshold = FMath::DegreesToRadians(AlignmentNotify->SteeringSettings.AngleThreshold);
+			
+			FTransform PredictedTransform;
+			float PredictionTime = StartTime;
+			float TargetTime = EndTime;
+
+			for(int i=0;i<NumFrames; i++)
+			{
+				FVector PrevPosition = PredictedTransform.GetLocation();
+				FTransform RootMotionThisFrame = ExtractRootMotionHelper(AnimationAsset, MirrorDataTable, bIsMirrored, PredictionTime, FrameTime, false);
+				
+				PredictedTransform = RootMotionThisFrame * PredictedTransform;
+				AnimTrajectoryData.Trajectory[i] = PredictedTransform;
+				PredictionTime += FrameTime;
+
+				AnimTrajectoryData.TranslationCurve[i] = RootMotionThisFrame.GetTranslation().Length();
+				if (i>0)
+				{
+					AnimTrajectoryData.TranslationCurve[i] += AnimTrajectoryData.TranslationCurve[i-1]; 
+				}
+				
+				AnimTrajectoryData.RotationCurve[i] = fabs(RootMotionThisFrame.GetRotation().GetAngle());
+				if (i>0)
+				{
+					AnimTrajectoryData.RotationCurve[i] += AnimTrajectoryData.RotationCurve[i-1]; 
+				}
+			}
+
+			// normalize curves;
+			if (AnimTrajectoryData.TranslationCurve.Last() > UE_SMALL_NUMBER)
+			{
+				if (AnimTrajectoryData.TranslationCurve.Last() < UE_SMALL_NUMBER)
+				{
+					AnimTrajectoryData.TranslationCurve.Reset();
+				}
+				else
+				{
+					for(float& Value : AnimTrajectoryData.TranslationCurve)
+					{
+						Value /= AnimTrajectoryData.TranslationCurve.Last();
+					}
+				}
+			}
+			
+			if (AnimTrajectoryData.RotationCurve.Last() > UE_SMALL_NUMBER)
+			{
+				if (AnimTrajectoryData.RotationCurve.Last() < UE_SMALL_NUMBER)
+				{
+					AnimTrajectoryData.RotationCurve.Reset();
+				}
+				else
+				{
+					for(float& Value : AnimTrajectoryData.RotationCurve)
 					{
-						if (!bPredictedRootMotionValid)
+						Value /= AnimTrajectoryData.RotationCurve.Last();
+					}
+				}
+			}
+
+			FVector UnWarpedPreviousPosition;
+			FVector WarpedPreviousPosition;
+
+			if (AlignmentNotify->bEnableSteering && AlignmentNotify->SteeringSettings.bEnableSmoothing)
+			{
+				FilteredSteeringTarget = FQuat::Identity;
+				TargetSmoothingState.Reset();
+			}
+
+			WarpedTrajectory.SetNum(AnimTrajectoryData.Trajectory.Num());
+			
+			FTransform InverseLastFrame = AnimTrajectoryData.Trajectory.Last().Inverse();
+			// Translation warping + steering
+			for(int i=0;i<NumFrames; i++)
+			{
+				float Weight = GetWeight(StartTime + FrameTime*i, AlignmentNotify->TranslationWarpingCurve);
+
+				FTransform TransformFromRoot = AnimTrajectoryData.Trajectory[i] * RootBoneTransform;
+				FTransform TransformFromTarget = AnimTrajectoryData.Trajectory[i] * InverseLastFrame * TargetTransform;
+				
+				FVector OldPosition = TransformFromRoot.GetTranslation();
+				FVector UnWarpedDelta = OldPosition - UnWarpedPreviousPosition;
+				UnWarpedPreviousPosition = OldPosition;
+				
+				FVector NewPosition = FMath::Lerp(OldPosition, TransformFromTarget.GetTranslation(), Weight);
+				FVector WarpedDelta = NewPosition - WarpedPreviousPosition;
+				WarpedPreviousPosition = NewPosition;
+
+				WarpedTrajectory[i].SetTranslation(NewPosition);
+				WarpedTrajectory[i].SetRotation(TransformFromRoot.GetRotation());
+
+				if (i > 0 && AlignmentNotify->bEnableSteering)
+				{
+					FQuat OldRotation = TransformFromRoot.GetRotation();
+					FQuat DirectionChange = FQuat::FindBetweenVectors(UnWarpedDelta, WarpedDelta);
+
+					if (AlignmentNotify->SteeringSettings.bEnableSmoothing)
+					{
+						
+						if (DirectionChange.GetAngle() < SteeringAngleThreshold)
 						{
-							// if we didn't already predict root motion above, or we predicted to a time offset, then compute the root motion to the end of the notify here
-							const float PredictionDelta = EndTime - CurrentTime;
-							PredictedRootMotion = AnimationAsset->ExtractRootMotion(FAnimExtractContext(static_cast<double>(CurrentTime), true, FDeltaTimeRecord(PredictionDelta), false));
+							FilteredSteeringTarget = UKismetMathLibrary::QuaternionSpringInterp(FilteredSteeringTarget, DirectionChange, TargetSmoothingState,
+															AlignmentNotify->SteeringSettings.SmoothStiffness, AlignmentNotify->SteeringSettings.SmoothDamping, DeltaTime, 1, 0, true);
 						}
-						
-						const float RootMotionDistanceRemaining = PredictedRootMotion.GetLocation().Length();
-						if (RootMotionDistanceRemaining > UE_SMALL_NUMBER)
+                        
+                        DirectionChange = FilteredSteeringTarget;
+						WarpedTrajectory[i].SetRotation(OldRotation*DirectionChange);
+					}
+					else
+					{
+						if (DirectionChange.GetAngle() < SteeringAngleThreshold)
 						{
-							const float RootMotionDistanceThisFrame = ThisFrameRootMotionTransform.GetLocation().Length();
-							Weight = RootMotionDistanceThisFrame / RootMotionDistanceRemaining;
+							WarpedTrajectory[i].SetRotation(OldRotation*DirectionChange);
 						}
 					}
+					
 				}
-    			
-    			AlignmentTransform = AlignmentNotify->AlignOffset * AlignmentTransform;
+			}
+			
+			// Rotation warping
+			for(int i=0;i<NumFrames; i++)
+			{
+				float Weight = GetWeight(StartTime + FrameTime*i, AlignmentNotify->RotationWarpingCurve);
+				
+				FQuat OldRotation = WarpedTrajectory[i].GetRotation();
+				FTransform TransformFromTarget = AnimTrajectoryData.Trajectory[i] * InverseLastFrame * TargetTransform;
+				
+				WarpedTrajectory[i].SetRotation(FQuat::Slerp(OldRotation, TransformFromTarget.GetRotation(), Weight));
+			}
+			
+		}
+	}
+	
+	const UE::Anim::IAnimRootMotionProvider* RootMotionProvider = UE::Anim::IAnimRootMotionProvider::Get();
+	ensureMsgf(RootMotionProvider, TEXT("Alignment expected a valid root motion delta provider interface."));
     		
-    			FQuat CurrentRotation = ThisFrameRootMotionTransform.GetRotation();
-    			CurrentRotation = CurrentRotation * FMath::Lerp(FQuat::Identity, AlignmentTransform.GetRotation().Inverse() * TargetTransform->GetRotation(), Weight);
-    			
-    			FVector CurrentTranslationWS = RootBoneTransform.TransformVector(ThisFrameRootMotionTransform.GetTranslation());
-    			CurrentTranslationWS += (TargetTransform->GetLocation() - AlignmentTransform.GetLocation())* Weight;
-
-    			ThisFrameRootMotionTransform.SetRotation(CurrentRotation);
-    			ThisFrameRootMotionTransform.SetTranslation(RootBoneTransform.InverseTransformVector(CurrentTranslationWS));
-
-    			
-    			RootMotionProvider->OverrideRootMotion(ThisFrameRootMotionTransform, Output.CustomAttributes);
-    		}
-    	}
-    }
+	if (RootMotionProvider)
+	{
+		float Frame = (CurrentTime - StartTime) / FrameTime;
+
+		int32 LowerFrame = FMath::Clamp(floor(Frame), 0, WarpedTrajectory.Num() -1);
+		int32 UpperFrame = FMath::Min(WarpedTrajectory.Num() -1 , LowerFrame + 1);
+		float Alpha = Frame - LowerFrame;
+
+		FTransform WorldTransform = WarpedTrajectory[LowerFrame];
+		WorldTransform.BlendWith(WarpedTrajectory[UpperFrame], Alpha);
+
+		RootMotionProvider->OverrideRootMotion(WorldTransform.GetRelativeTransform(RootBoneTransform), Output.CustomAttributes);
+		
+		// unwarped trajectory relative to starting transform
+		FTransform PreviousTransform = AnimTrajectoryData.Trajectory[0] * StartingRootTransform;
+		for(int i=1;i<AnimTrajectoryData.Trajectory.Num(); i++)
+		{
+			FTransform TransformFromRoot = AnimTrajectoryData.Trajectory[i] * StartingRootTransform;
+			UE_VLOG_SEGMENT(Output.AnimInstanceProxy->GetAnimInstanceObject(), "Alignment", Display, PreviousTransform.GetLocation(), TransformFromRoot.GetLocation(), i%2 == 0 ? FColor::Yellow : FColor::Red, TEXT(""));
+			PreviousTransform = TransformFromRoot;
+		}
+		
+		// unwarped trajectory relative to target transform
+		FTransform InverseLastFrame = AnimTrajectoryData.Trajectory.Last().Inverse();
+		PreviousTransform = AnimTrajectoryData.Trajectory[0] * InverseLastFrame * TargetTransform;
+		for(int i=1;i<AnimTrajectoryData.Trajectory.Num(); i++)
+		{
+			FTransform TransformFromTarget = AnimTrajectoryData.Trajectory[i] * InverseLastFrame * TargetTransform;
+			UE_VLOG_SEGMENT(Output.AnimInstanceProxy->GetAnimInstanceObject(), "Alignment", Display, PreviousTransform.GetLocation(), TransformFromTarget.GetLocation(), i%2 == 0 ? FColor::Yellow : FColor::Red, TEXT(""));
+			PreviousTransform = TransformFromTarget;
+		}
+		
+		// the warped trajectory
+		for(int i=1; i<WarpedTrajectory.Num(); i++)
+		{
+			UE_VLOG_SEGMENT(Output.AnimInstanceProxy->GetAnimInstanceObject(), "Alignment", Display, WarpedTrajectory[i-1].GetLocation(), WarpedTrajectory[i].GetLocation(), i%2 == 0 ? FColor::Green : FColor::Blue, TEXT(""));
+		}
+
+		// a dot representing our current position on the trajectory 
+		UE_VLOG_SPHERE(Output.AnimInstanceProxy->GetAnimInstanceObject(), "Alignment", Display, WorldTransform.GetLocation(),  1, FColor::Red, TEXT(""));
+	}
 }
 
diff --git a/Engine/Plugins/Experimental/EvaluationNotifies/Source/Runtime/Private/EvaluationNotifies/AnimNotifyState_TwoBoneIK.cpp b/Engine/Plugins/Experimental/EvaluationNotifies/Source/Runtime/Private/EvaluationNotifies/AnimNotifyState_TwoBoneIK.cpp
--- a/Engine/Plugins/Experimental/EvaluationNotifies/Source/Runtime/Private/EvaluationNotifies/AnimNotifyState_TwoBoneIK.cpp
+++ b/Engine/Plugins/Experimental/EvaluationNotifies/Source/Runtime/Private/EvaluationNotifies/AnimNotifyState_TwoBoneIK.cpp
@@ -42,7 +42,7 @@ static FAutoConsoleVariableRef CVarEnableTwoBoneIKNotify(
 	g_bEnableTwoBoneIKNotify,
 	TEXT("Enable Two BoneIK Notify"));
 
-void FTwoBoneIKNotifyInstance::Update(const UAnimSequenceBase* AnimationAsset, float CurrentTime, float DeltaTime, 
+void FTwoBoneIKNotifyInstance::Update(const UAnimSequenceBase* AnimationAsset, float CurrentTime, float DeltaTime, bool bIsMirrored, const UMirrorDataTable* MirrorDataTable,
                                    FTransform& InRootBoneTransform, const TMap<FName, FTransform>& NamedTransforms, FComponentSpacePoseContext& Output, TArray<FBoneTransform>& OutBoneTransforms)
 {
 	if (!g_bEnableTwoBoneIKNotify)
diff --git a/Engine/Plugins/Experimental/EvaluationNotifies/Source/Runtime/Public/EvaluationNotifies/AnimNode_EvaluationNotifies.h b/Engine/Plugins/Experimental/EvaluationNotifies/Source/Runtime/Public/EvaluationNotifies/AnimNode_EvaluationNotifies.h
--- a/Engine/Plugins/Experimental/EvaluationNotifies/Source/Runtime/Public/EvaluationNotifies/AnimNode_EvaluationNotifies.h
+++ b/Engine/Plugins/Experimental/EvaluationNotifies/Source/Runtime/Public/EvaluationNotifies/AnimNode_EvaluationNotifies.h
@@ -19,7 +19,7 @@ struct FEvaluationNotifyInstance
 	GENERATED_BODY()
 	
 	virtual void Start(const UAnimSequenceBase* AnimationAsset) {}
-	virtual void Update(const UAnimSequenceBase* AnimationAsset, float CurrentTime, float DeltaTime,
+	virtual void Update(const UAnimSequenceBase* AnimationAsset, float CurrentTime, float DeltaTime, bool bIsMirrored, const UMirrorDataTable* MirrorDataTable,
 		FTransform& RootBoneTransform, const TMap<FName, FTransform>& NamedTransforms, FComponentSpacePoseContext& Output, TArray<FBoneTransform>& OutBoneTransforms) {}
 	virtual void End() {}
 
@@ -45,6 +45,14 @@ struct EVALUATIONNOTIFIESRUNTIME_API FAnimNode_EvaluationNotifies : public FAnim
 	UPROPERTY(EditAnywhere, Transient, BlueprintReadWrite, Category=Evaluation, meta=(PinShownByDefault))
 	float CurrentAnimAssetTime = 0.f;
 	
+	// Is the current anim asset mirrored
+	UPROPERTY(EditAnywhere, Transient, BlueprintReadWrite, Category=Evaluation, meta=(PinShownByDefault))
+	bool CurrentAnimAssetMirrored = false;
+	
+	// If bMirrored is set, MirrorDataTable will be used for mirroring the CurrentAnimAsset during prediction
+	UPROPERTY(EditAnywhere, BlueprintReadWrite, Category = Evaluation, meta = (PinShownByDefault))
+	TObjectPtr<UMirrorDataTable> MirrorDataTable;
+	
 	// Current playback time in seconds of the CurrentAnimAsset
 	UPROPERTY(EditAnywhere, Transient, BlueprintReadWrite, Category=Evaluation, meta=(PinShownByDefault))
 	TMap<FName, FTransform> NamedTransforms;
diff --git a/Engine/Plugins/Experimental/EvaluationNotifies/Source/Runtime/Public/EvaluationNotifies/AnimNotifyState_Alignment.h b/Engine/Plugins/Experimental/EvaluationNotifies/Source/Runtime/Public/EvaluationNotifies/AnimNotifyState_Alignment.h
--- a/Engine/Plugins/Experimental/EvaluationNotifies/Source/Runtime/Public/EvaluationNotifies/AnimNotifyState_Alignment.h
+++ b/Engine/Plugins/Experimental/EvaluationNotifies/Source/Runtime/Public/EvaluationNotifies/AnimNotifyState_Alignment.h
@@ -8,22 +8,62 @@
 UENUM()
 enum class EAlignmentWeightCurveType
 {
-	FromRootMotion,
+	FromRootMotionTranslation,
+	FromRootMotionRotation,
 	Linear,
 	EaseIn,
 	EaseOut,
-	EaseInOut
+	EaseInOut,
+	Instant,
+	DoNotWarp,
 };
 
-UENUM()
-enum class EAlignmentPredictionMode
+USTRUCT()
+struct FAlignmentWarpCurve
+{
+	GENERATED_BODY()
+	
+	// Type of curve to interpolate using
+	UPROPERTY(EditAnywhere, Category = Properties)
+	EAlignmentWeightCurveType CurveType = EAlignmentWeightCurveType::FromRootMotionTranslation;
+
+	// Time to start warping, as a ratio of the duration of the Notify
+	UPROPERTY(EditAnywhere, Category = "Properties", meta = (UIMin = "0.0", UIMax = "1.0", SliderExponent = "1.0"))
+	float StartRatio = 0.f;
+
+	// Time to complete warping, as a ratio of the duration of the Notify
+	UPROPERTY(EditAnywhere, Category = "Properties", meta = (UIMin = "0.0", UIMax = "1.0", SliderExponent = "1.0"))
+	float EndRatio = 1.f;
+};
+
+USTRUCT()
+struct FAlignmentSteeringSettings
+{
+	GENERATED_BODY()
+
+	// Enable smoothing of the steering target orientation, to avoid instant orientation changes
+	UPROPERTY(EditAnywhere, Category=Properties)
+	bool bEnableSmoothing = true;
+
+	// Sprint stiffness for smoothing
+	UPROPERTY(EditAnywhere, Category=Propeties, meta = (EditCondition=bEnableSmoothing))
+	float SmoothStiffness = 300;
+
+	// Spring damping factor for smoothing
+	UPROPERTY(EditAnywhere, Category=Properties, meta = (EditCondition=bEnableSmoothing))
+	float SmoothDamping = 1;
+	
+	// When the warped movement direction differs from the animated movement direction by more than this threshold, steering will be disabled.
+	UPROPERTY(EditAnywhere, Category = Properties)
+	float AngleThreshold = 135;
+};
+
+struct AnimTrajectoryData
 {
-	// Root Motion data from the AnimSequence will be used to predict the animated target transform, 
-	RootMotion,
-	// A Bone in the AnimSequence has been animated to represent the target transform
-	AlignBone,
-	// Warping will be applied every frame without considering the future movement in the AnimSequence
-	None,
+	FTransform TargetTransform;
+	TArray<FTransform> Trajectory;
+	TArray<float> TranslationCurve;
+	TArray<float> RotationCurve;
 };
 
 UCLASS(BlueprintType, DisplayName="Alignment")
@@ -32,32 +72,35 @@ class UNotifyState_Alignment : public UAnimNotifyState
 	GENERATED_BODY()
 	
 	public:
+
+	UPROPERTY(EditAnywhere, Category = Blending)
+	FAlignmentWarpCurve TranslationWarpingCurve;
 	
-	// Curve used to control the rate at which we blend into alignment
-	UPROPERTY(EditAnywhere, Category = Properties)
-	EAlignmentWeightCurveType WeightCurve = EAlignmentWeightCurveType::FromRootMotion;
+	UPROPERTY(EditAnywhere, Category = Blending)
+	FAlignmentWarpCurve RotationWarpingCurve = { };
 	
 	// Named transform to align to
-    UPROPERTY(EditAnywhere, Category = Properties)
+    UPROPERTY(EditAnywhere, Category = TargetTransform)
     FName TransformName;
 
 	// Offset from Root (or AlignBone) to Align to the target transform
-	UPROPERTY(EditAnywhere, Category = Properties)
+	UPROPERTY(EditAnywhere, Category = TargetTransform)
 	FTransform AlignOffset;
 
-	// Method for predicting how the animation will move the character relative to the alignment target
-	UPROPERTY(EditAnywhere, Category = Properties)
-	EAlignmentPredictionMode PredictionMode;
-
-	// Time in seconds after the end of the NotifyState when the root should be in alignment with the target (warping will happen over the notify duration, non-warped root motion will complete the alignment)
-	UPROPERTY(EditAnywhere, Category = Properties, meta = (EditCondition = "PredictionMode == EAlignmentPredictionMode::RootMotion"))
-	float PostNotifyTargetTime = 0.f;
-
-	// Bone to Align to target transform
-	UPROPERTY(EditAnywhere, Category = Properties, meta = (EditCondition = "PredictionMode == EAlignmentPredictionMode::AlignBone"))
+	// Optional bone to Align to target transform
+	UPROPERTY(EditAnywhere, Category = TargetTransform)
 	FBoneReference AlignBone;
+		
+	// // a Dynamic Target might move and so warping needs to be updated every frame
+	// UPROPERTY(EditAnywhere, Category = TargetTransform)
+	// bool bDynamicTarget;
 
+	// Steering will rotate the character based on the difference between the animated movement direction and the warped movement direction, to keep the character facing their direction of movement
+	UPROPERTY(EditAnywhere, Category = Steering)
+	bool bEnableSteering = false;
 	
+	UPROPERTY(EditAnywhere, Category = Steering, meta=(EditCondition=bEnableSteering))
+	FAlignmentSteeringSettings SteeringSettings;
 };
 
 USTRUCT()
@@ -66,15 +109,20 @@ struct FAlignmentNotifyInstance : public FEvaluationNotifyInstance
 	GENERATED_BODY()
 
 	virtual void Start(const UAnimSequenceBase* AnimationAsset) override;
-	virtual void Update(const UAnimSequenceBase* AnimationAsset, float CurrentTime, float DeltaTime,
+	virtual void Update(const UAnimSequenceBase* AnimationAsset, float CurrentTime, float DeltaTime, bool bIsMirrored, const UMirrorDataTable* MirrorDataTable,
 		FTransform& RootBoneTransform, const TMap<FName, FTransform>& NamedTransforms, FComponentSpacePoseContext& Output, TArray<FBoneTransform>& OutBoneTransforms) override;
 	
+	float GetWeight(float CurrentTime, const FAlignmentWarpCurve& WarpCurve) const;
 	
-	float GetWeightThisFrame(float CurrentTime, float DeltaTime);
-
-	float PreviousCurveWeight = 0.0;
 	FBoneReference AlignBone;
-	float PostNotifyTargetTime = 0.f;
-	EAlignmentPredictionMode PredictionMode;
-	EAlignmentWeightCurveType WeightCurve;
+	bool bFirstFrame = true;
+
+	FTransform StartingRootTransform;
+	FTransform TargetTransform;
+	
+	FQuat FilteredSteeringTarget = FQuat::Identity;
+	FQuaternionSpringState TargetSmoothingState;
+
+	TArray<FTransform> WarpedTrajectory;
+	AnimTrajectoryData AnimTrajectoryData;
 };
\ No newline at end of file
diff --git a/Engine/Plugins/Experimental/EvaluationNotifies/Source/Runtime/Public/EvaluationNotifies/AnimNotifyState_TwoBoneIK.h b/Engine/Plugins/Experimental/EvaluationNotifies/Source/Runtime/Public/EvaluationNotifies/AnimNotifyState_TwoBoneIK.h
--- a/Engine/Plugins/Experimental/EvaluationNotifies/Source/Runtime/Public/EvaluationNotifies/AnimNotifyState_TwoBoneIK.h
+++ b/Engine/Plugins/Experimental/EvaluationNotifies/Source/Runtime/Public/EvaluationNotifies/AnimNotifyState_TwoBoneIK.h
@@ -94,7 +94,7 @@ struct FTwoBoneIKNotifyInstance : public FEvaluationNotifyInstance
 	GENERATED_BODY()
 
 	virtual void Start(const UAnimSequenceBase* AnimationAsset) override;
-	virtual void Update(const UAnimSequenceBase* AnimationAsset, float CurrentTime, float DeltaTime,
+	virtual void Update(const UAnimSequenceBase* AnimationAsset, float CurrentTime, float DeltaTime, bool bIsMirrored, const UMirrorDataTable* MirrorDataTable,
 		FTransform& RootBoneTransform, const TMap<FName, FTransform>& NamedTransforms, FComponentSpacePoseContext& Output, TArray<FBoneTransform>& OutBoneTransforms) override;
 
 	FBoneReference IKBone;
