In-graph warping alignment fixes for partially executed notifies

diff --git a/Engine/Plugins/Experimental/EvaluationNotifies/Source/Runtime/Private/EvaluationNotifies/AnimNode_EvaluationNotifies.cpp b/Engine/Plugins/Experimental/EvaluationNotifies/Source/Runtime/Private/EvaluationNotifies/AnimNode_EvaluationNotifies.cpp
--- a/Engine/Plugins/Experimental/EvaluationNotifies/Source/Runtime/Private/EvaluationNotifies/AnimNode_EvaluationNotifies.cpp
+++ b/Engine/Plugins/Experimental/EvaluationNotifies/Source/Runtime/Private/EvaluationNotifies/AnimNode_EvaluationNotifies.cpp
@@ -54,6 +54,20 @@ void FAnimNode_EvaluationNotifies::GetEvaluationNotifiesFromAnimation(const UAni
 void FAnimNode_EvaluationNotifies::UpdateInternal(const FAnimationUpdateContext& Context)
 {
 	FAnimNode_SkeletalControlBase::UpdateInternal(Context);
+	
+	const bool bJustBecameRelevant = !UpdateCounter.WasSynchronizedCounter(Context.AnimInstanceProxy->GetUpdateCounter());
+	UpdateCounter.SynchronizeWith(Context.AnimInstanceProxy->GetUpdateCounter());
+
+	if (bJustBecameRelevant)
+	{
+		// clear state on become relevant in case there's any left over evaluation notifies from last time
+		CurrentSequence = nullptr;
+		Tags.Empty();
+	}
+
+	TRACE_ANIM_NODE_VALUE(Context, TEXT("Asset"), CurrentAnimAsset);
+	TRACE_ANIM_NODE_VALUE(Context, TEXT("Playback Time"), CurrentAnimAssetTime);
+	TRACE_ANIM_NODE_VALUE(Context, TEXT("Mirrored"), CurrentAnimAssetMirrored);
 
 	if (UE::AnimationWarping::FRootOffsetProvider* RootOffsetProvider = Context.GetMessage<UE::AnimationWarping::FRootOffsetProvider>())
 	{
diff --git a/Engine/Plugins/Experimental/EvaluationNotifies/Source/Runtime/Private/EvaluationNotifies/AnimNotifyState_Alignment.cpp b/Engine/Plugins/Experimental/EvaluationNotifies/Source/Runtime/Private/EvaluationNotifies/AnimNotifyState_Alignment.cpp
--- a/Engine/Plugins/Experimental/EvaluationNotifies/Source/Runtime/Private/EvaluationNotifies/AnimNotifyState_Alignment.cpp
+++ b/Engine/Plugins/Experimental/EvaluationNotifies/Source/Runtime/Private/EvaluationNotifies/AnimNotifyState_Alignment.cpp
@@ -65,16 +65,18 @@ float FAlignmentNotifyInstance::GetWeight(float CurrentTime, const FAlignmentWar
 	if (WarpCurve.CurveType == EAlignmentWeightCurveType::FromRootMotionTranslation && !AnimTrajectoryData.TranslationCurve.IsEmpty())
 	{
 		const float Duration = (EndTime - StartTime);
-		return SampleCurve(CurrentTime - StartTime, WarpCurve.StartRatio * Duration, WarpCurve.EndRatio * Duration, AnimTrajectoryData.TranslationCurve);
+		const float StartCurveSampleTime = FMath::Max(WarpCurve.StartRatio * Duration, ActualStartTime) - ActualStartTime;
+		return SampleCurve(CurrentTime - ActualStartTime, StartCurveSampleTime, WarpCurve.EndRatio * Duration - ActualStartTime, AnimTrajectoryData.TranslationCurve);
 	}
 	else if (WarpCurve.CurveType == EAlignmentWeightCurveType::FromRootMotionRotation && !AnimTrajectoryData.RotationCurve.IsEmpty())
 	{
 		const float Duration = (EndTime - StartTime);
-		return SampleCurve(CurrentTime - StartTime, WarpCurve.StartRatio * Duration, WarpCurve.EndRatio * Duration, AnimTrajectoryData.RotationCurve);
+		const float StartCurveSampleTime = FMath::Max(WarpCurve.StartRatio * Duration, ActualStartTime) - ActualStartTime;
+		return SampleCurve(CurrentTime - ActualStartTime, StartCurveSampleTime, WarpCurve.EndRatio * Duration - ActualStartTime, AnimTrajectoryData.RotationCurve);
 	}
 	else
 	{
-		const float StartTimeWithRatio = FMath::Lerp(StartTime, EndTime, WarpCurve.StartRatio);
+		const float StartTimeWithRatio = FMath::Max(FMath::Lerp(StartTime, EndTime, WarpCurve.StartRatio), ActualStartTime);
 		const float EndTimeWithRatio = FMath::Lerp(StartTime, EndTime, WarpCurve.EndRatio);
 			
 		const float CurrentRelativeTime = FMath::Clamp((CurrentTime-StartTimeWithRatio)/(EndTimeWithRatio-StartTimeWithRatio), 0,1);
@@ -114,6 +116,7 @@ void FAlignmentNotifyInstance::Update(const UAnimSequenceBase* AnimationAsset, f
 	if (bFirstFrame)
 	{
 		bFirstFrame = false;
+		ActualStartTime = CurrentTime;
 		
 		const FBoneContainer& RequiredBones = Output.AnimInstanceProxy->GetRequiredBones();	
 		AlignBone.Initialize(RequiredBones);
@@ -156,7 +159,7 @@ void FAlignmentNotifyInstance::Update(const UAnimSequenceBase* AnimationAsset, f
 			float SteeringAngleThreshold = FMath::DegreesToRadians(AlignmentNotify->SteeringSettings.AngleThreshold);
 			
 			FTransform PredictedTransform;
-			float PredictionTime = StartTime;
+			float PredictionTime = CurrentTime;
 			float TargetTime = EndTime;
 
 			for(int i=0;i<NumFrames; i++)
@@ -227,7 +230,7 @@ void FAlignmentNotifyInstance::Update(const UAnimSequenceBase* AnimationAsset, f
 			// Translation warping + steering
 			for(int i=0;i<NumFrames; i++)
 			{
-				float Weight = GetWeight(StartTime + FrameTime*i, AlignmentNotify->TranslationWarpingCurve);
+				float Weight = GetWeight(ActualStartTime + FrameTime*i, AlignmentNotify->TranslationWarpingCurve);
 
 				FTransform TransformFromRoot = AnimTrajectoryData.Trajectory[i] * RootBoneTransform;
 				FTransform TransformFromTarget = AnimTrajectoryData.Trajectory[i] * InverseLastFrame * TargetTransform;
@@ -274,7 +277,7 @@ void FAlignmentNotifyInstance::Update(const UAnimSequenceBase* AnimationAsset, f
 			// Rotation warping
 			for(int i=0;i<NumFrames; i++)
 			{
-				float Weight = GetWeight(StartTime + FrameTime*i, AlignmentNotify->RotationWarpingCurve);
+				float Weight = GetWeight(ActualStartTime + FrameTime*i, AlignmentNotify->RotationWarpingCurve);
 				
 				FQuat OldRotation = WarpedTrajectory[i].GetRotation();
 				FTransform TransformFromTarget = AnimTrajectoryData.Trajectory[i] * InverseLastFrame * TargetTransform;
@@ -287,14 +290,16 @@ void FAlignmentNotifyInstance::Update(const UAnimSequenceBase* AnimationAsset, f
 	
 	const UE::Anim::IAnimRootMotionProvider* RootMotionProvider = UE::Anim::IAnimRootMotionProvider::Get();
 	ensureMsgf(RootMotionProvider, TEXT("Alignment expected a valid root motion delta provider interface."));
+	
     		
-	if (RootMotionProvider)
+	if (RootMotionProvider && !WarpedTrajectory.IsEmpty())
 	{
-		float Frame = (CurrentTime - StartTime) / FrameTime;
+		float Frame = (CurrentTime - ActualStartTime) / FrameTime;
 
 		int32 LowerFrame = FMath::Clamp(floor(Frame), 0, WarpedTrajectory.Num() -1);
 		int32 UpperFrame = FMath::Min(WarpedTrajectory.Num() -1 , LowerFrame + 1);
 		float Alpha = Frame - LowerFrame;
+		
 
 		FTransform WorldTransform = WarpedTrajectory[LowerFrame];
 		WorldTransform.BlendWith(WarpedTrajectory[UpperFrame], Alpha);
diff --git a/Engine/Plugins/Experimental/EvaluationNotifies/Source/Runtime/Public/EvaluationNotifies/AnimNode_EvaluationNotifies.h b/Engine/Plugins/Experimental/EvaluationNotifies/Source/Runtime/Public/EvaluationNotifies/AnimNode_EvaluationNotifies.h
--- a/Engine/Plugins/Experimental/EvaluationNotifies/Source/Runtime/Public/EvaluationNotifies/AnimNode_EvaluationNotifies.h
+++ b/Engine/Plugins/Experimental/EvaluationNotifies/Source/Runtime/Public/EvaluationNotifies/AnimNode_EvaluationNotifies.h
@@ -81,4 +81,7 @@ private:
 	TObjectPtr<const UAnimSequenceBase> CurrentSequence = nullptr;
 
 	static TMap<UClass*, UScriptStruct*> NotifyEvaluationHandlerMap;
+
+	// Update Counter for detecting being relevant
+	FGraphTraversalCounter UpdateCounter;
 };
diff --git a/Engine/Plugins/Experimental/EvaluationNotifies/Source/Runtime/Public/EvaluationNotifies/AnimNotifyState_Alignment.h b/Engine/Plugins/Experimental/EvaluationNotifies/Source/Runtime/Public/EvaluationNotifies/AnimNotifyState_Alignment.h
--- a/Engine/Plugins/Experimental/EvaluationNotifies/Source/Runtime/Public/EvaluationNotifies/AnimNotifyState_Alignment.h
+++ b/Engine/Plugins/Experimental/EvaluationNotifies/Source/Runtime/Public/EvaluationNotifies/AnimNotifyState_Alignment.h
@@ -116,6 +116,7 @@ struct FAlignmentNotifyInstance : public FEvaluationNotifyInstance
 	
 	FBoneReference AlignBone;
 	bool bFirstFrame = true;
+	float ActualStartTime = 0;
 
 	FTransform StartingRootTransform;
 	FTransform TargetTransform;
