EvaluationNotifies anim next support

diff --git a/Engine/Plugins/Experimental/AnimNextAnimGraph/Source/AnimNextAnimGraph/Private/Traits/SequencePlayer.cpp b/Engine/Plugins/Experimental/AnimNextAnimGraph/Source/AnimNextAnimGraph/Private/Traits/SequencePlayer.cpp
--- a/Engine/Plugins/Experimental/AnimNextAnimGraph/Source/AnimNextAnimGraph/Private/Traits/SequencePlayer.cpp
+++ b/Engine/Plugins/Experimental/AnimNextAnimGraph/Source/AnimNextAnimGraph/Private/Traits/SequencePlayer.cpp
@@ -186,6 +186,7 @@ namespace UE::AnimNext
 		if (const UAnimSequence* AnimSeq = InstanceData->AnimSequence.Get())
 		{
 			FAnimTickRecord TickRecord;
+			TickRecord.TimeAccumulator = &StartPosition;
 			TickRecord.bLooping = bLooping;
 			FAnimNotifyContext NotifyContext(TickRecord);
 			AnimSeq->GetAnimNotifies(StartPosition, Duration, NotifyContext);
diff --git a/Engine/Plugins/Experimental/EvaluationNotifies/EvaluationNotifies.uplugin b/Engine/Plugins/Experimental/EvaluationNotifies/EvaluationNotifies.uplugin
--- a/Engine/Plugins/Experimental/EvaluationNotifies/EvaluationNotifies.uplugin
+++ b/Engine/Plugins/Experimental/EvaluationNotifies/EvaluationNotifies.uplugin
@@ -28,6 +28,18 @@
 		{
 			"Name": "AnimationWarping",
 			"Enabled" : true
+		},
+		{
+			"Name": "RigVM",
+			"Enabled" : true
+		},
+		{
+			"Name": "AnimNext",
+			"Enabled" : true
+		},
+		{
+			"Name": "AnimNextAnimGraph",
+			"Enabled" : true
 		}
 	]
 }
\ No newline at end of file
diff --git a/Engine/Plugins/Experimental/EvaluationNotifies/Source/Runtime/EvaluationNotifiesRuntime.Build.cs b/Engine/Plugins/Experimental/EvaluationNotifies/Source/Runtime/EvaluationNotifiesRuntime.Build.cs
--- a/Engine/Plugins/Experimental/EvaluationNotifies/Source/Runtime/EvaluationNotifiesRuntime.Build.cs
+++ b/Engine/Plugins/Experimental/EvaluationNotifies/Source/Runtime/EvaluationNotifiesRuntime.Build.cs
@@ -18,7 +18,10 @@ namespace UnrealBuildTool.Rules
                 "Core",
                 "CoreUObject",
                 "Engine",
-                "AnimationWarpingRuntime"
+                "AnimationWarpingRuntime",
+                "RigVM",
+                "AnimNext",
+                "AnimNextAnimGraph", 
             });
 		}
 	}
diff --git a/Engine/Plugins/Experimental/EvaluationNotifies/Source/Runtime/Private/AnimNext/AnimNextAlignment.cpp b/Engine/Plugins/Experimental/EvaluationNotifies/Source/Runtime/Private/AnimNext/AnimNextAlignment.cpp
new file mode 100644
--- /dev/null
+++ b/Engine/Plugins/Experimental/EvaluationNotifies/Source/Runtime/Private/AnimNext/AnimNextAlignment.cpp
@@ -0,0 +1,344 @@
+// Copyright Epic Games, Inc. All Rights Reserved.
+
+#include "AnimNextAlignment.h"
+
+#include "AnimNextAnimGraphSettings.h"
+#include <AHEasing/easing.h>
+
+#include "Animation/AnimRootMotionProvider.h"
+#include "Animation/AnimSequence.h"
+#include "EvaluationNotifies/AnimNotifyState_Alignment.h"
+#include "EvaluationVM/EvaluationVM.h"
+#include "EvaluationVM/KeyframeState.h"
+#include "Graph/AnimNextGraphInstance.h"
+#include "TraitCore/NodeInstance.h"
+#include "VisualLogger/VisualLogger.h"
+
+void FEvaluationNotify_AlignmentInstance::Start()
+{
+	const UNotifyState_Alignment* AlignmentNotify = Cast<UNotifyState_Alignment>(AnimNotify);
+	AlignBone = AlignmentNotify->AlignBone;
+	bFirstFrame = true;
+}
+
+namespace
+{
+	constexpr float FrameTime = 1.0f/30.0f;
+}
+
+namespace
+{
+	float SampleCurve(float SampleTime, float StartTime, float EndTime, const TArray<float>& CurveData)
+	{
+		if (SampleTime <= StartTime)
+		{
+			return 0.0f;
+		}
+		if (SampleTime >= EndTime)
+		{
+			return 1.0;
+		}
+	
+		// do nearest neighbor sampling for simplicity and performance
+		const int StartFrame = FMath::Clamp(round(StartTime / FrameTime), 0, CurveData.Num()-1);
+		const int EndFrame = FMath::Clamp(round(EndTime / FrameTime), 0, CurveData.Num()-1);
+		const int SampleFrame = FMath::Clamp(round(SampleTime / FrameTime), 0, CurveData.Num()-1);
+
+		return (CurveData[SampleFrame] - CurveData[StartFrame])/(CurveData[EndFrame] - CurveData[StartFrame]);
+	}
+}
+
+float FEvaluationNotify_AlignmentInstance::GetWeight(float Time, const FAlignmentWarpCurve& WarpCurve) const 
+{
+	UNotifyState_Alignment* AlignmentNotify = Cast<UNotifyState_Alignment>(AnimNotify);
+	float Weight;
+	
+
+	if (WarpCurve.CurveType == EAlignmentWeightCurveType::FromRootMotionTranslation && !AnimTrajectoryData.TranslationCurve.IsEmpty())
+	{
+		const float Duration = (EndTime - StartTime);
+		const float StartCurveSampleTime = FMath::Max(WarpCurve.StartRatio * Duration, ActualStartTime) - ActualStartTime;
+		return SampleCurve(Time - ActualStartTime, StartCurveSampleTime, WarpCurve.EndRatio * Duration - ActualStartTime, AnimTrajectoryData.TranslationCurve);
+	}
+	else if (WarpCurve.CurveType == EAlignmentWeightCurveType::FromRootMotionRotation && !AnimTrajectoryData.RotationCurve.IsEmpty())
+	{
+		const float Duration = (EndTime - StartTime);
+		const float StartCurveSampleTime = FMath::Max(WarpCurve.StartRatio * Duration, ActualStartTime) - ActualStartTime;
+		return SampleCurve(Time - ActualStartTime, StartCurveSampleTime, WarpCurve.EndRatio * Duration - ActualStartTime, AnimTrajectoryData.RotationCurve);
+	}
+	else
+	{
+		const float StartTimeWithRatio = FMath::Max(FMath::Lerp(StartTime, EndTime, WarpCurve.StartRatio), ActualStartTime);
+		const float EndTimeWithRatio = FMath::Lerp(StartTime, EndTime, WarpCurve.EndRatio);
+			
+		const float CurrentRelativeTime = FMath::Clamp((Time-StartTimeWithRatio)/(EndTimeWithRatio-StartTimeWithRatio), 0,1);
+		
+		switch(WarpCurve.CurveType)
+		{
+			case EAlignmentWeightCurveType::EaseIn:
+				Weight = CubicEaseIn(CurrentRelativeTime);
+				break;
+			case EAlignmentWeightCurveType::EaseOut:
+				Weight = CubicEaseOut(CurrentRelativeTime);
+				break;
+			case EAlignmentWeightCurveType::EaseInOut:
+				Weight = CubicEaseInOut(CurrentRelativeTime);
+				break;
+			case EAlignmentWeightCurveType::Instant:
+				Weight = 1.f;
+				break;
+			case EAlignmentWeightCurveType::DoNotWarp:
+				Weight = 0.f;
+				break;
+			default:
+				Weight = CurrentRelativeTime;
+		}
+	}
+
+	
+	return Weight;
+}
+
+FTransform GetModelSpaceTransform(const UE::AnimNext::FLODPoseStack& Pose, FBoneIndexType SkeletonBoneIndex)
+{
+	const TArrayView<const FBoneIndexType> SkeletonToPoseIndexMap = Pose.GetSkeletonBoneIndexToLODBoneIndexMap();
+	const FBoneIndexType PoseIndex = SkeletonToPoseIndexMap[SkeletonBoneIndex];
+	FTransform BoneTransform = Pose.LocalTransforms[PoseIndex];
+						
+	FBoneIndexType ParentSkeletonIndex = Pose.GetSkeletonAsset()->GetReferenceSkeleton().GetParentIndex(SkeletonBoneIndex);
+	while(ParentSkeletonIndex != (FBoneIndexType)-1)
+	{
+		const FBoneIndexType ParentPoseIndex = SkeletonToPoseIndexMap[ParentSkeletonIndex];
+		BoneTransform = BoneTransform * Pose.LocalTransforms[ParentPoseIndex];
+		ParentSkeletonIndex = Pose.GetSkeletonAsset()->GetReferenceSkeleton().GetParentIndex(ParentSkeletonIndex);
+	}
+
+	return BoneTransform;
+}
+
+void FEvaluationNotify_AlignmentInstance::Update(UE::AnimNext::FEvaluationNotifiesTrait::FInstanceData& TraitInstanceData, UE::AnimNext::FEvaluationVM& VM)
+{
+	UNotifyState_Alignment* AlignmentNotify = Cast<UNotifyState_Alignment>(AnimNotify);
+	
+	if (const TUniquePtr<UE::AnimNext::FKeyframeState>* Keyframe = VM.PeekValue<TUniquePtr<UE::AnimNext::FKeyframeState>>(UE::AnimNext::KEYFRAME_STACK_NAME, 0))
+	{
+		if (bFirstFrame)
+		{
+			bFirstFrame = false;
+			ActualStartTime = CurrentTime;
+
+			const UE::AnimNext::FLODPoseStack& Pose = Keyframe->Get()->Pose;
+			AlignBone.Initialize(Pose.GetSkeletonAsset());
+
+			if (TraitInstanceData.DataInterface->GetVariable(AlignmentNotify->TransformName, TargetTransform) == EPropertyBagResult::Success)
+			{
+				static const FBox UnitBox(FVector(-10, -10, -10), FVector(10, 10, 10));
+				UE_VLOG_OBOX(TraitInstanceData.DebugObject, "Alignment", Display, UnitBox, TargetTransform.ToMatrixWithScale(), FColor::Blue, TEXT(""));
+			
+				// get alignment transform
+				const float PredictionDelta = EndTime - CurrentTime;
+			
+
+				if (AlignBone.HasValidSetup())
+				{
+					// get alignment bone relative to predicted root motion end point, and remove that as an offset to the alignment target
+					// Alignment bone is expected to be at the alignment end point at the beginning of the Notify window  (we don't have a good way to prediuct the component space transform of it in the future)
+				 	
+					FTransform PredictedRootMotion = TraitInstanceData.OnExtractRootMotionAttribute.Execute(CurrentTime, PredictionDelta, false);
+
+					FTransform AlignBoneTransform = GetModelSpaceTransform(Pose, AlignBone.BoneIndex);
+				 	TargetTransform = AlignBoneTransform.Inverse() * PredictedRootMotion * TargetTransform;
+				 }
+			
+				TargetTransform = AlignmentNotify->AlignOffset * TargetTransform;
+			
+				StartingRootTransform = TraitInstanceData.RootBoneTransform;
+			
+				// extract root motion trajectory todo: this should be cached and reused
+			
+				int NumFrames = (EndTime - CurrentTime) / FrameTime;
+
+				AnimTrajectoryData.Trajectory.SetNum(NumFrames);
+				AnimTrajectoryData.TranslationCurve.SetNum(NumFrames);
+				AnimTrajectoryData.RotationCurve.SetNum(NumFrames);
+			
+				float SteeringAngleThreshold = FMath::DegreesToRadians(AlignmentNotify->SteeringSettings.AngleThreshold);
+			
+				FTransform PredictedTransform;
+				float PredictionTime = CurrentTime;
+				float TargetTime = EndTime;
+
+				for(int i=0;i<NumFrames; i++)
+				{
+					FTransform RootMotionThisFrame = TraitInstanceData.OnExtractRootMotionAttribute.Execute(PredictionTime, FrameTime, false);
+				
+					PredictedTransform = RootMotionThisFrame * PredictedTransform;
+					AnimTrajectoryData.Trajectory[i] = PredictedTransform;
+					PredictionTime += FrameTime;
+
+					AnimTrajectoryData.TranslationCurve[i] = RootMotionThisFrame.GetTranslation().Length();
+					if (i>0)
+					{
+						AnimTrajectoryData.TranslationCurve[i] += AnimTrajectoryData.TranslationCurve[i-1]; 
+					}
+				
+					AnimTrajectoryData.RotationCurve[i] = fabs(RootMotionThisFrame.GetRotation().GetAngle());
+					if (i>0)
+					{
+						AnimTrajectoryData.RotationCurve[i] += AnimTrajectoryData.RotationCurve[i-1]; 
+					}
+				}
+
+				// normalize curves;
+				if (AnimTrajectoryData.TranslationCurve.Last() > UE_SMALL_NUMBER)
+				{
+					if (AnimTrajectoryData.TranslationCurve.Last() < UE_SMALL_NUMBER)
+					{
+						AnimTrajectoryData.TranslationCurve.Reset();
+					}
+					else
+					{
+						for(float& Value : AnimTrajectoryData.TranslationCurve)
+						{
+							Value /= AnimTrajectoryData.TranslationCurve.Last();
+						}
+					}
+				}
+			
+				if (AnimTrajectoryData.RotationCurve.Last() > UE_SMALL_NUMBER)
+				{
+					if (AnimTrajectoryData.RotationCurve.Last() < UE_SMALL_NUMBER)
+					{
+						AnimTrajectoryData.RotationCurve.Reset();
+					}
+					else
+					{
+						for(float& Value : AnimTrajectoryData.RotationCurve)
+						{
+							Value /= AnimTrajectoryData.RotationCurve.Last();
+						}
+					}
+				}
+
+				FVector UnWarpedPreviousPosition;
+				FVector WarpedPreviousPosition;
+
+				if (AlignmentNotify->bEnableSteering && AlignmentNotify->SteeringSettings.bEnableSmoothing)
+				{
+					FilteredSteeringTarget = FQuat::Identity;
+					TargetSmoothingState.Reset();
+				}
+
+				WarpedTrajectory.SetNum(AnimTrajectoryData.Trajectory.Num());
+			
+				FTransform InverseLastFrame = AnimTrajectoryData.Trajectory.Last().Inverse();
+				// Translation warping + steering
+				for(int i=0;i<NumFrames; i++)
+				{
+					float Weight = GetWeight(ActualStartTime + FrameTime*i, AlignmentNotify->TranslationWarpingCurve);
+
+					FTransform TransformFromRoot = AnimTrajectoryData.Trajectory[i] * TraitInstanceData.RootBoneTransform;
+					FTransform TransformFromTarget = AnimTrajectoryData.Trajectory[i] * InverseLastFrame * TargetTransform;
+				
+					FVector OldPosition = TransformFromRoot.GetTranslation();
+					FVector UnWarpedDelta = OldPosition - UnWarpedPreviousPosition;
+					UnWarpedPreviousPosition = OldPosition;
+				
+					FVector NewPosition = FMath::Lerp(OldPosition, TransformFromTarget.GetTranslation(), Weight);
+					FVector WarpedDelta = NewPosition - WarpedPreviousPosition;
+					WarpedPreviousPosition = NewPosition;
+
+					WarpedTrajectory[i].SetTranslation(NewPosition);
+					WarpedTrajectory[i].SetRotation(TransformFromRoot.GetRotation());
+
+					if (i > 0 && AlignmentNotify->bEnableSteering)
+					{
+						FQuat OldRotation = TransformFromRoot.GetRotation();
+						FQuat DirectionChange = FQuat::FindBetweenVectors(UnWarpedDelta, WarpedDelta);
+
+						if (AlignmentNotify->SteeringSettings.bEnableSmoothing)
+						{
+						
+							if (DirectionChange.GetAngle() < SteeringAngleThreshold)
+							{
+								FilteredSteeringTarget = UKismetMathLibrary::QuaternionSpringInterp(FilteredSteeringTarget, DirectionChange, TargetSmoothingState,
+																AlignmentNotify->SteeringSettings.SmoothStiffness, AlignmentNotify->SteeringSettings.SmoothDamping, TraitInstanceData.DeltaTime, 1, 0, true);
+							}
+                        
+							DirectionChange = FilteredSteeringTarget;
+							WarpedTrajectory[i].SetRotation(OldRotation*DirectionChange);
+						}
+						else
+						{
+							if (DirectionChange.GetAngle() < SteeringAngleThreshold)
+							{
+								WarpedTrajectory[i].SetRotation(OldRotation*DirectionChange);
+							}
+						}
+					
+					}
+				}
+			
+				// Rotation warping
+				for(int i=0;i<NumFrames; i++)
+				{
+					float Weight = GetWeight(ActualStartTime + FrameTime*i, AlignmentNotify->RotationWarpingCurve);
+				
+					FQuat OldRotation = WarpedTrajectory[i].GetRotation();
+					FTransform TransformFromTarget = AnimTrajectoryData.Trajectory[i] * InverseLastFrame * TargetTransform;
+				
+					WarpedTrajectory[i].SetRotation(FQuat::Slerp(OldRotation, TransformFromTarget.GetRotation(), Weight));
+				}
+			
+			}
+		}
+	
+		const UE::Anim::IAnimRootMotionProvider* RootMotionProvider = UE::Anim::IAnimRootMotionProvider::Get();
+		ensureMsgf(RootMotionProvider, TEXT("Alignment expected a valid root motion delta provider interface."));
+	
+    		
+		if (RootMotionProvider && !WarpedTrajectory.IsEmpty())
+		{
+			float Frame = (CurrentTime - ActualStartTime) / FrameTime;
+
+			int32 LowerFrame = FMath::Clamp(floor(Frame), 0, WarpedTrajectory.Num() -1);
+			int32 UpperFrame = FMath::Min(WarpedTrajectory.Num() -1 , LowerFrame + 1);
+			float Alpha = Frame - LowerFrame;
+		
+
+			FTransform WorldTransform = WarpedTrajectory[LowerFrame];
+			WorldTransform.BlendWith(WarpedTrajectory[UpperFrame], Alpha);
+
+			RootMotionProvider->OverrideRootMotion(WorldTransform.GetRelativeTransform(TraitInstanceData.RootBoneTransform), Keyframe->Get()->Attributes);
+		
+			// unwarped trajectory relative to starting transform
+			FTransform PreviousTransform = AnimTrajectoryData.Trajectory[0] * StartingRootTransform;
+			for(int i=1;i<AnimTrajectoryData.Trajectory.Num(); i++)
+			{
+				FTransform TransformFromRoot = AnimTrajectoryData.Trajectory[i] * StartingRootTransform;
+				UE_VLOG_SEGMENT(TraitInstanceData.DebugObject, "Alignment", Display, PreviousTransform.GetLocation(), TransformFromRoot.GetLocation(), i%2 == 0 ? FColor::Yellow : FColor::Red, TEXT(""));
+				PreviousTransform = TransformFromRoot;
+			}
+		
+			// unwarped trajectory relative to target transform
+			FTransform InverseLastFrame = AnimTrajectoryData.Trajectory.Last().Inverse();
+			PreviousTransform = AnimTrajectoryData.Trajectory[0] * InverseLastFrame * TargetTransform;
+			for(int i=1;i<AnimTrajectoryData.Trajectory.Num(); i++)
+			{
+				FTransform TransformFromTarget = AnimTrajectoryData.Trajectory[i] * InverseLastFrame * TargetTransform;
+				UE_VLOG_SEGMENT(TraitInstanceData.DebugObject, "Alignment", Display, PreviousTransform.GetLocation(), TransformFromTarget.GetLocation(), i%2 == 0 ? FColor::Yellow : FColor::Red, TEXT(""));
+				PreviousTransform = TransformFromTarget;
+			}
+		
+			// the warped trajectory
+			for(int i=1; i<WarpedTrajectory.Num(); i++)
+			{
+				UE_VLOG_SEGMENT(TraitInstanceData.DebugObject, "Alignment", Display, WarpedTrajectory[i-1].GetLocation(), WarpedTrajectory[i].GetLocation(), i%2 == 0 ? FColor::Green : FColor::Blue, TEXT(""));
+			}
+
+			// a dot representing our current position on the trajectory 
+			UE_VLOG_SPHERE(TraitInstanceData.DebugObject, "Alignment", Display, WorldTransform.GetLocation(),  1, FColor::Red, TEXT(""));
+		}
+	}
+}
+
diff --git a/Engine/Plugins/Experimental/EvaluationNotifies/Source/Runtime/Private/AnimNext/AnimNextAlignment.h b/Engine/Plugins/Experimental/EvaluationNotifies/Source/Runtime/Private/AnimNext/AnimNextAlignment.h
new file mode 100644
--- /dev/null
+++ b/Engine/Plugins/Experimental/EvaluationNotifies/Source/Runtime/Private/AnimNext/AnimNextAlignment.h
@@ -0,0 +1,48 @@
+// Copyright Epic Games, Inc. All Rights Reserved.
+
+#pragma once
+
+#include "CoreMinimal.h"
+#include "EvaluationNotifiesTrait.h"
+#include "DataInterface/AnimNextDataInterfaceInstance.h"
+#include "EvaluationNotifies/AnimNotifyState_Alignment.h"
+#include "Kismet/KismetMathLibrary.h"
+#include "TraitCore/Trait.h"
+#include "TraitCore/TraitSharedData.h"
+#include "TraitInterfaces/IAttributeProvider.h"
+#include "TraitInterfaces/IEvaluate.h"
+#include "TraitInterfaces/ITimeline.h"
+#include "TraitInterfaces/IUpdate.h"
+#include "AnimNextAlignment.generated.h"
+
+struct AlignmentAnimTrajectoryData
+{
+	FTransform TargetTransform;
+	TArray<FTransform> Trajectory;
+	TArray<float> TranslationCurve;
+	TArray<float> RotationCurve;
+};
+
+
+USTRUCT()
+struct FEvaluationNotify_AlignmentInstance : public FEvaluationNotify_BaseInstance
+{
+	GENERATED_BODY()
+	virtual void Start() override;
+   	virtual void Update(UE::AnimNext::FEvaluationNotifiesTrait::FInstanceData& InstanceData, UE::AnimNext::FEvaluationVM& VM) override;
+	
+	float GetWeight(float Time, const FAlignmentWarpCurve& WarpCurve) const; 
+
+	FBoneReference AlignBone;
+	bool bFirstFrame = true;
+	float ActualStartTime = 0;
+
+	FTransform StartingRootTransform;
+	FTransform TargetTransform;
+	
+	FQuat FilteredSteeringTarget = FQuat::Identity;
+	FQuaternionSpringState TargetSmoothingState;
+
+	TArray<FTransform> WarpedTrajectory;
+	AlignmentAnimTrajectoryData AnimTrajectoryData;
+};
\ No newline at end of file
diff --git a/Engine/Plugins/Experimental/EvaluationNotifies/Source/Runtime/Private/AnimNext/EvaluationNotifiesTrait.cpp b/Engine/Plugins/Experimental/EvaluationNotifies/Source/Runtime/Private/AnimNext/EvaluationNotifiesTrait.cpp
new file mode 100644
--- /dev/null
+++ b/Engine/Plugins/Experimental/EvaluationNotifies/Source/Runtime/Private/AnimNext/EvaluationNotifiesTrait.cpp
@@ -0,0 +1,205 @@
+// Copyright Epic Games, Inc. All Rights Reserved.
+
+#include "EvaluationNotifiesTrait.h"
+
+#include "AnimNextAnimGraphSettings.h"
+// #include "AnimNextWarpingLog.h"
+#include <AHEasing/easing.h>
+
+#include "Animation/AnimRootMotionProvider.h"
+#include "Animation/AnimSequence.h"
+#include "EvaluationNotifies/AnimNotifyState_Alignment.h"
+#include "EvaluationVM/EvaluationVM.h"
+#include "Graph/AnimNextGraphInstance.h"
+#include "TraitCore/NodeInstance.h"
+#include "TraitInterfaces/IBlendStack.h"
+#include "TraitInterfaces/IGraphFactory.h"
+#include "TraitInterfaces/INotifySource.h"
+#include "TraitInterfaces/ITimeline.h"
+
+static const FName RootBoneTransformName = "RootBoneTransform";
+
+////////////////////////////////////////////////////////////////////////////////////////////////////////////////////////
+// FEvaluationNotifiesTrait
+
+namespace UE::AnimNext
+{
+	TMap<UClass*, UScriptStruct*> FEvaluationNotifiesTrait::NotifyEvaluationHandlerMap;
+
+	AUTO_REGISTER_ANIM_TRAIT(FEvaluationNotifiesTrait)
+
+	#define TRAIT_INTERFACE_ENUMERATOR(GeneratorMacro) \
+		GeneratorMacro(ITimeline) \
+		GeneratorMacro(IUpdate) \
+		GeneratorMacro(IEvaluate) \
+
+	// Trait required interfaces implementation boilerplate
+    #define TRAIT_REQUIRED_INTERFACE_ENUMERATOR(GeneratorMacroRequired) \
+    	GeneratorMacroRequired(ITimeline) \
+    	GeneratorMacroRequired(INotifySource) \
+
+	// Trait implementation boilerplate
+	GENERATE_ANIM_TRAIT_IMPLEMENTATION(FEvaluationNotifiesTrait, TRAIT_INTERFACE_ENUMERATOR, NULL_ANIM_TRAIT_INTERFACE_ENUMERATOR, NULL_ANIM_TRAIT_EVENT_ENUMERATOR)
+	#undef TRAIT_INTERFACE_ENUMERATOR
+
+	void FEvaluationNotifiesTrait::AdvanceBy(FExecutionContext& Context, const TTraitBinding<ITimeline>& Binding, float DeltaTime) const
+	{
+		TTraitBinding<ITimeline> TimelineTrait;
+		ensure(Binding.GetStackInterfaceSuper(TimelineTrait));
+
+		TTraitBinding<INotifySource> NotifySourceTrait;
+		ensure(Binding.GetStackInterfaceSuper(NotifySourceTrait));
+
+		// Get current state from the stack, advance time, then get the delta state
+		const FTimelineState PreAdvanceState = TimelineTrait.GetState(Context);
+		TimelineTrait.AdvanceBy(Context, DeltaTime);
+		const FTimelineDelta Delta = TimelineTrait.GetDelta(Context);
+
+		// Query for notifies
+		TArray<FAnimNotifyEventReference> Notifies;
+		NotifySourceTrait.GetNotifies(Context, PreAdvanceState.GetPosition(), Delta.GetDeltaTime(), PreAdvanceState.IsLooping(), Notifies);
+
+		FInstanceData* InstanceData = Binding.GetInstanceData<FInstanceData>();
+
+		for(const FAnimNotifyEventReference& NotifyData : Notifies)
+		{
+			if(const FAnimNotifyEvent* NotifyEvent = NotifyData.GetNotify())
+			{
+				if (UAnimNotifyState* Notify = Cast<UAnimNotifyState>(NotifyEvent->NotifyStateClass))
+				{
+					if (UScriptStruct** HandlerType = NotifyEvaluationHandlerMap.Find(Notify->GetClass()))
+					{
+						if (FInstancedStruct* FoundInstance = InstanceData->EvaluationNotifies.FindByPredicate([NotifyEvent](const FInstancedStruct ExistingInstance)
+						{
+							// check if we already have an instance for this event
+							return ExistingInstance.Get<FEvaluationNotify_BaseInstance>().NotifyEvent == NotifyEvent;
+						}))
+						{
+							FoundInstance->GetMutable<FEvaluationNotify_BaseInstance>().CurrentTime = NotifyData.GetCurrentAnimationTime() + DeltaTime;
+						}
+						else
+						{
+							FInstancedStruct& Struct = InstanceData->EvaluationNotifies.AddDefaulted_GetRef();
+							Struct.InitializeAs(*HandlerType);
+							FEvaluationNotify_BaseInstance& Instance = Struct.GetMutable<FEvaluationNotify_BaseInstance>();
+							Instance.NotifyEvent = NotifyEvent;
+							Instance.AnimNotify = Notify;
+							Instance.StartTime = NotifyEvent->GetTriggerTime();
+							Instance.EndTime = NotifyEvent->GetEndTriggerTime();
+							Instance.CurrentTime = NotifyData.GetCurrentAnimationTime();
+						}
+					}
+				}
+			}
+		}
+	}
+
+	void FEvaluationNotifiesTrait::PreUpdate(FUpdateTraversalContext& Context, const TTraitBinding<IUpdate>& Binding, const FTraitUpdateState& TraitState) const
+	{
+		FInstanceData* InstanceData = Binding.GetInstanceData<FInstanceData>();
+		check(InstanceData);
+
+		InstanceData->DeltaTime = TraitState.GetDeltaTime();
+
+		IUpdate::PreUpdate(Context, Binding, TraitState);
+
+	}
+
+	void FEvaluationNotifiesTrait::PostEvaluate(FEvaluateTraversalContext& Context, const TTraitBinding<IEvaluate>& Binding) const
+	{
+		IEvaluate::PostEvaluate(Context, Binding);
+
+		const FSharedData* SharedData = Binding.GetSharedData<FSharedData>();
+		check(SharedData);
+
+		FInstanceData* InstanceData = Binding.GetInstanceData<FInstanceData>();
+		check(InstanceData);
+
+#if ENABLE_ANIM_DEBUG
+		InstanceData->DebugObject = Context.GetHostObject();
+#endif
+
+		// Get Root motion provider
+		{
+			TTraitBinding<IAttributeProvider> AttributeTrait;
+			if (!Binding.GetStackInterface(AttributeTrait))
+			{
+				// UE_LOG(LogAnimNextWarping, Error, TEXT("FSteeringTrait::PostEvaluate, missing IAttributeProvider"));
+				return;
+			}
+
+			InstanceData->OnExtractRootMotionAttribute = AttributeTrait.GetOnExtractRootMotionAttribute(Context);
+		}
+		
+		InstanceData->DataInterface = Binding.GetTraitPtr().GetNodeInstance()->GetOwner().GetHost();
+		InstanceData->DataInterface->GetVariable(RootBoneTransformName, InstanceData->RootBoneTransform);
+
+		Context.AppendTask(FAnimNextEvaluationNotifiesTask::Make(InstanceData, SharedData));
+	}
+
+
+	void FEvaluationNotifiesTrait::RegisterEvaluationHandler(UClass* NotifyType, UScriptStruct* HandlerType)
+    {
+    	NotifyEvaluationHandlerMap.Add(NotifyType, HandlerType);
+    }
+    
+    void FEvaluationNotifiesTrait::UnregisterEvaluationHandler(UClass* NotifyType)
+    {
+    	NotifyEvaluationHandlerMap.Remove(NotifyType);
+    }
+	
+} // namespace UE::AnimNext
+
+
+////////////////////////////////////////////////////////////////////////////////////////////////////////////////////////
+// FAnimNextEvaluationNotifiesTask
+
+FAnimNextEvaluationNotifiesTask FAnimNextEvaluationNotifiesTask::Make(UE::AnimNext::FEvaluationNotifiesTrait::FInstanceData* InstanceData, const UE::AnimNext::FEvaluationNotifiesTrait::FSharedData* SharedData)
+{
+	FAnimNextEvaluationNotifiesTask Task;
+	Task.InstanceData = InstanceData;
+	Task.SharedData = SharedData;
+	return Task;
+}
+
+void FAnimNextEvaluationNotifiesTask::Execute(UE::AnimNext::FEvaluationVM& VM) const
+{
+	using namespace UE::AnimNext;
+
+	// Execute EvaluationNotifies
+	if (InstanceData->DeltaTime > 0)
+	{
+		for(auto& NotifyInstanceData : InstanceData->EvaluationNotifies)
+		{
+			FEvaluationNotify_BaseInstance& NotifyInstance = NotifyInstanceData.GetMutable<FEvaluationNotify_BaseInstance>();
+			
+			if (NotifyInstance.StartTime <= NotifyInstance.CurrentTime && NotifyInstance.EndTime >= NotifyInstance.CurrentTime)
+			{
+				if (!NotifyInstance.bActive)
+				{
+					NotifyInstance.bActive = true;
+					NotifyInstance.Start();
+				}
+
+				NotifyInstance.Update(*InstanceData, VM);
+			}
+			else
+			{
+				if (NotifyInstance.bActive)
+				{
+					NotifyInstance.bActive = false;
+					NotifyInstance.End();
+				}
+			}
+			
+		}
+
+		InstanceData->EvaluationNotifies.RemoveAll([](const FInstancedStruct& Data)
+		{
+			return Data.Get<FEvaluationNotify_BaseInstance>().bActive == false;
+		});
+	
+	}
+
+}
+
diff --git a/Engine/Plugins/Experimental/EvaluationNotifies/Source/Runtime/Private/AnimNext/EvaluationNotifiesTrait.h b/Engine/Plugins/Experimental/EvaluationNotifies/Source/Runtime/Private/AnimNext/EvaluationNotifiesTrait.h
new file mode 100644
--- /dev/null
+++ b/Engine/Plugins/Experimental/EvaluationNotifies/Source/Runtime/Private/AnimNext/EvaluationNotifiesTrait.h
@@ -0,0 +1,124 @@
+// Copyright Epic Games, Inc. All Rights Reserved.
+
+#pragma once
+
+#include "CoreMinimal.h"
+#include "DataInterface/AnimNextDataInterfaceInstance.h"
+#include "Kismet/KismetMathLibrary.h"
+#include "TraitCore/Trait.h"
+#include "TraitCore/TraitSharedData.h"
+#include "TraitInterfaces/IAttributeProvider.h"
+#include "TraitInterfaces/IEvaluate.h"
+#include "TraitInterfaces/ITimeline.h"
+#include "TraitInterfaces/IUpdate.h"
+#include "EvaluationNotifiesTrait.generated.h"
+
+struct FEvaluationNotify_BaseInstance;
+
+/**
+ * Data needed to execute steering
+ * 
+ * Some steering data such as current anim asset / playback time is acquired via trait stack interfaces
+ */
+USTRUCT(meta = (DisplayName = "EvaluationNotifies"))
+struct FEvaluationNotifiesTraitSharedData : public FAnimNextTraitSharedData
+{
+	GENERATED_BODY()
+
+	// Latent pin support boilerplate
+	#define TRAIT_LATENT_PROPERTIES_ENUMERATOR(GeneratorMacro) \
+
+	GENERATE_TRAIT_LATENT_PROPERTIES(FEvaluationNotifiesTraitSharedData, TRAIT_LATENT_PROPERTIES_ENUMERATOR)
+	#undef TRAIT_LATENT_PROPERTIES_ENUMERATOR
+};
+
+namespace UE::AnimNext
+{
+
+	
+/**
+ * this trait will run evaluation time code for notifies which have a registered handler
+ */
+struct EVALUATIONNOTIFIESRUNTIME_API FEvaluationNotifiesTrait : FAdditiveTrait, ITimeline, IUpdate, IEvaluate
+{
+	// @TODO: Hash we use use FNV1A32. Update macro to mention that and also comment if we still need it (?)
+	DECLARE_ANIM_TRAIT(FEvaluationNotifiesTrait, 0x4aa0ccf2, FAdditiveTrait)
+
+	using FSharedData = FEvaluationNotifiesTraitSharedData;
+
+	struct FInstanceData : FTrait::FInstanceData
+	{
+		/** Delta in seconds between updates, populated during PreUpdate */
+		float DeltaTime = 0.f;
+
+		/** Callback provided by attribute trait on stack to evaluate root motion at a later time */
+		FOnExtractRootMotionAttribute OnExtractRootMotionAttribute = FOnExtractRootMotionAttribute();
+
+		/** Last root bone transform sampled */
+		FTransform RootBoneTransform = FTransform::Identity;
+
+		// the evaluation notifies extracted from the current anim sequence
+		TArray<FInstancedStruct> EvaluationNotifies;
+
+		const FAnimNextDataInterfaceInstance* DataInterface = nullptr;
+
+#if ENABLE_ANIM_DEBUG
+		// an object for attaching debug draws to
+		const UObject* DebugObject;
+#endif
+	};
+	
+	
+	// ITimeline impl
+	virtual void AdvanceBy(FExecutionContext& Context, const TTraitBinding<ITimeline>& Binding, float DeltaTime) const override;
+	
+	// IUpdate impl 
+	virtual void PreUpdate(FUpdateTraversalContext& Context, const TTraitBinding<IUpdate>& Binding, const FTraitUpdateState& TraitState) const override;
+
+	// IEvaluate impl
+	virtual void PostEvaluate(FEvaluateTraversalContext& Context, const TTraitBinding<IEvaluate>& Binding) const override;
+
+	static void RegisterEvaluationHandler(UClass* NotifyType, UScriptStruct* Handler);
+	static void UnregisterEvaluationHandler(UClass* NotifyType);
+	
+private:
+	static TMap<UClass*, UScriptStruct*> NotifyEvaluationHandlerMap;
+};
+
+} // namespace UE::AnimNext
+
+/** Task to run EvaluationNotifies on VM */
+USTRUCT()
+struct FAnimNextEvaluationNotifiesTask : public FAnimNextEvaluationTask
+{
+	GENERATED_BODY()
+
+	DECLARE_ANIM_EVALUATION_TASK(FAnimNextEvaluationNotifiesTask)
+
+	static FAnimNextEvaluationNotifiesTask Make(UE::AnimNext::FEvaluationNotifiesTrait::FInstanceData* InstanceData, const UE::AnimNext::FEvaluationNotifiesTrait::FSharedData* SharedData);
+
+	// Task entry point
+	virtual void Execute(UE::AnimNext::FEvaluationVM& VM) const override;
+
+	UE::AnimNext::FEvaluationNotifiesTrait::FInstanceData* InstanceData = nullptr;
+	const UE::AnimNext::FEvaluationNotifiesTrait::FSharedData* SharedData = nullptr;
+};
+
+USTRUCT()
+struct FEvaluationNotify_BaseInstance
+{
+	GENERATED_BODY()
+	virtual void Start() {}
+	virtual void Update(UE::AnimNext::FEvaluationNotifiesTrait::FInstanceData& InstanceData, UE::AnimNext::FEvaluationVM& VM) {}
+	virtual void End() {}
+	
+	virtual ~FEvaluationNotify_BaseInstance() {};
+
+	const FAnimNotifyEvent* NotifyEvent;
+	TObjectPtr<UAnimNotifyState> AnimNotify = nullptr;
+	float StartTime = 0.f;
+	float EndTime = 0.f;
+	
+	float CurrentTime = 0.f;
+	bool bActive = false;
+};
\ No newline at end of file
diff --git a/Engine/Plugins/Experimental/EvaluationNotifies/Source/Runtime/Private/EvaluationNotifiesRuntimeModule.cpp b/Engine/Plugins/Experimental/EvaluationNotifies/Source/Runtime/Private/EvaluationNotifiesRuntimeModule.cpp
--- a/Engine/Plugins/Experimental/EvaluationNotifies/Source/Runtime/Private/EvaluationNotifiesRuntimeModule.cpp
+++ b/Engine/Plugins/Experimental/EvaluationNotifies/Source/Runtime/Private/EvaluationNotifiesRuntimeModule.cpp
@@ -2,9 +2,11 @@
 
 #include "EvaluationNotifiesRuntimeModule.h"
 #include "Animation/AnimRootMotionProvider.h"
+#include "AnimNext/EvaluationNotifiesTrait.h"
 #include "Features/IModularFeatures.h"
 #include "Modules/ModuleManager.h"
-#include "EvaluationNotifies/AnimNode_EvaluationNotifies.h"
+
+#include "AnimNext/AnimNextAlignment.h"
 #include "EvaluationNotifies/AnimNotifyState_Alignment.h"
 #include "EvaluationNotifies/AnimNotifyState_TwoBoneIK.h"
 
@@ -34,12 +36,16 @@ namespace UE
 		{
 			FAnimNode_EvaluationNotifies::RegisterEvaluationHandler(UNotifyState_Alignment::StaticClass(), FAlignmentNotifyInstance::StaticStruct());
 			FAnimNode_EvaluationNotifies::RegisterEvaluationHandler(UNotifyState_TwoBoneIK::StaticClass(), FTwoBoneIKNotifyInstance::StaticStruct());
+
+			AnimNext::FEvaluationNotifiesTrait::RegisterEvaluationHandler(UNotifyState_Alignment::StaticClass(), FEvaluationNotify_AlignmentInstance::StaticStruct());
 		}
 
 		void FModule::ShutdownModule()
 		{
 			FAnimNode_EvaluationNotifies::UnregisterEvaluationHandler(UNotifyState_Alignment::StaticClass());
 			FAnimNode_EvaluationNotifies::UnregisterEvaluationHandler(UNotifyState_TwoBoneIK::StaticClass());
+			
+			AnimNext::FEvaluationNotifiesTrait::UnregisterEvaluationHandler(UNotifyState_Alignment::StaticClass());
 		}
 	}
 }
