Unity build fix

diff --git a/Engine/Plugins/Experimental/EvaluationNotifies/Source/Runtime/Private/AnimNext/AnimNextAlignment.cpp b/Engine/Plugins/Experimental/EvaluationNotifies/Source/Runtime/Private/AnimNext/AnimNextAlignment.cpp
--- a/Engine/Plugins/Experimental/EvaluationNotifies/Source/Runtime/Private/AnimNext/AnimNextAlignment.cpp
+++ b/Engine/Plugins/Experimental/EvaluationNotifies/Source/Runtime/Private/AnimNext/AnimNextAlignment.cpp
@@ -21,13 +21,10 @@ void FEvaluationNotify_AlignmentInstance::Start()
 	bFirstFrame = true;
 }
 
-namespace
+namespace AnimNextAlignment
 {
 	constexpr float FrameTime = 1.0f/30.0f;
-}
-
-namespace
-{
+	
 	float SampleCurve(float SampleTime, float StartTime, float EndTime, const TArray<float>& CurveData)
 	{
 		if (SampleTime <= StartTime)
@@ -58,13 +55,13 @@ float FEvaluationNotify_AlignmentInstance::GetWeight(float Time, const FAlignmen
 	{
 		const float Duration = (EndTime - StartTime);
 		const float StartCurveSampleTime = FMath::Max(WarpCurve.StartRatio * Duration, ActualStartTime) - ActualStartTime;
-		return SampleCurve(Time - ActualStartTime, StartCurveSampleTime, WarpCurve.EndRatio * Duration - ActualStartTime, AnimTrajectoryData.TranslationCurve);
+		return AnimNextAlignment::SampleCurve(Time - ActualStartTime, StartCurveSampleTime, WarpCurve.EndRatio * Duration - ActualStartTime, AnimTrajectoryData.TranslationCurve);
 	}
 	else if (WarpCurve.CurveType == EAlignmentWeightCurveType::FromRootMotionRotation && !AnimTrajectoryData.RotationCurve.IsEmpty())
 	{
 		const float Duration = (EndTime - StartTime);
 		const float StartCurveSampleTime = FMath::Max(WarpCurve.StartRatio * Duration, ActualStartTime) - ActualStartTime;
-		return SampleCurve(Time - ActualStartTime, StartCurveSampleTime, WarpCurve.EndRatio * Duration - ActualStartTime, AnimTrajectoryData.RotationCurve);
+		return AnimNextAlignment::SampleCurve(Time - ActualStartTime, StartCurveSampleTime, WarpCurve.EndRatio * Duration - ActualStartTime, AnimTrajectoryData.RotationCurve);
 	}
 	else
 	{
@@ -156,7 +153,7 @@ void FEvaluationNotify_AlignmentInstance::Update(UE::AnimNext::FEvaluationNotifi
 			
 				// extract root motion trajectory todo: this should be cached and reused
 			
-				int NumFrames = (EndTime - CurrentTime) / FrameTime;
+				int NumFrames = (EndTime - CurrentTime) / AnimNextAlignment::FrameTime;
 
 				AnimTrajectoryData.Trajectory.SetNum(NumFrames);
 				AnimTrajectoryData.TranslationCurve.SetNum(NumFrames);
@@ -170,11 +167,11 @@ void FEvaluationNotify_AlignmentInstance::Update(UE::AnimNext::FEvaluationNotifi
 
 				for(int i=0;i<NumFrames; i++)
 				{
-					FTransform RootMotionThisFrame = TraitInstanceData.OnExtractRootMotionAttribute.Execute(PredictionTime, FrameTime, false);
+					FTransform RootMotionThisFrame = TraitInstanceData.OnExtractRootMotionAttribute.Execute(PredictionTime, AnimNextAlignment::FrameTime, false);
 				
 					PredictedTransform = RootMotionThisFrame * PredictedTransform;
 					AnimTrajectoryData.Trajectory[i] = PredictedTransform;
-					PredictionTime += FrameTime;
+					PredictionTime += AnimNextAlignment::FrameTime;
 
 					AnimTrajectoryData.TranslationCurve[i] = RootMotionThisFrame.GetTranslation().Length();
 					if (i>0)
@@ -235,7 +232,7 @@ void FEvaluationNotify_AlignmentInstance::Update(UE::AnimNext::FEvaluationNotifi
 				// Translation warping + steering
 				for(int i=0;i<NumFrames; i++)
 				{
-					float Weight = GetWeight(ActualStartTime + FrameTime*i, AlignmentNotify->TranslationWarpingCurve);
+					float Weight = GetWeight(ActualStartTime + AnimNextAlignment::FrameTime*i, AlignmentNotify->TranslationWarpingCurve);
 
 					FTransform TransformFromRoot = AnimTrajectoryData.Trajectory[i] * TraitInstanceData.RootBoneTransform;
 					FTransform TransformFromTarget = AnimTrajectoryData.Trajectory[i] * InverseLastFrame * TargetTransform;
@@ -282,7 +279,7 @@ void FEvaluationNotify_AlignmentInstance::Update(UE::AnimNext::FEvaluationNotifi
 				// Rotation warping
 				for(int i=0;i<NumFrames; i++)
 				{
-					float Weight = GetWeight(ActualStartTime + FrameTime*i, AlignmentNotify->RotationWarpingCurve);
+					float Weight = GetWeight(ActualStartTime + AnimNextAlignment::FrameTime*i, AlignmentNotify->RotationWarpingCurve);
 				
 					FQuat OldRotation = WarpedTrajectory[i].GetRotation();
 					FTransform TransformFromTarget = AnimTrajectoryData.Trajectory[i] * InverseLastFrame * TargetTransform;
@@ -299,7 +296,7 @@ void FEvaluationNotify_AlignmentInstance::Update(UE::AnimNext::FEvaluationNotifi
     		
 		if (RootMotionProvider && !WarpedTrajectory.IsEmpty())
 		{
-			float Frame = (CurrentTime - ActualStartTime) / FrameTime;
+			float Frame = (CurrentTime - ActualStartTime) / AnimNextAlignment::FrameTime;
 
 			int32 LowerFrame = FMath::Clamp(floor(Frame), 0, WarpedTrajectory.Num() -1);
 			int32 UpperFrame = FMath::Min(WarpedTrajectory.Num() -1 , LowerFrame + 1);
