Relative alignment mode

diff --git a/Engine/Plugins/Experimental/EvaluationNotifies/Source/Runtime/Private/AnimNext/AnimNextAlignment.cpp b/Engine/Plugins/Experimental/EvaluationNotifies/Source/Runtime/Private/AnimNext/AnimNextAlignment.cpp
--- a/Engine/Plugins/Experimental/EvaluationNotifies/Source/Runtime/Private/AnimNext/AnimNextAlignment.cpp
+++ b/Engine/Plugins/Experimental/EvaluationNotifies/Source/Runtime/Private/AnimNext/AnimNextAlignment.cpp
@@ -43,6 +43,17 @@ namespace AnimNextAlignment
 
 		return (CurveData[SampleFrame] - CurveData[StartFrame])/(CurveData[EndFrame] - CurveData[StartFrame]);
 	}
+	
+	void GetTransformForFrame(float Frame, const TArray<FTransform>& Trajectory, FTransform& OutTransform)
+	{
+		int32 LowerFrame = FMath::Clamp(floor(Frame), 0, Trajectory.Num() -1);
+		int32 UpperFrame = FMath::Min(Trajectory.Num() -1 , LowerFrame + 1);
+		float Alpha = Frame - LowerFrame;
+
+		OutTransform = Trajectory[LowerFrame];
+		OutTransform.BlendWith(Trajectory[UpperFrame], Alpha);
+	}
+
 }
 
 float FEvaluationNotify_AlignmentInstance::GetWeight(float Time, const FAlignmentWarpCurve& WarpCurve) const 
@@ -224,6 +235,9 @@ void FEvaluationNotify_AlignmentInstance::Update(UE::AnimNext::FEvaluationNotifi
 				{
 					FilteredSteeringTarget = FQuat::Identity;
 					TargetSmoothingState.Reset();
+
+					FilteredSteeringTarget = UKismetMathLibrary::QuaternionSpringInterp(FilteredSteeringTarget, FQuat::Identity, TargetSmoothingState,
+													AlignmentNotify->SteeringSettings.SmoothStiffness, AlignmentNotify->SteeringSettings.SmoothDamping, TraitInstanceData.DeltaTime, 1, 0, true);
 				}
 
 				WarpedTrajectory.SetNum(AnimTrajectoryData.Trajectory.Num());
@@ -297,16 +311,20 @@ void FEvaluationNotify_AlignmentInstance::Update(UE::AnimNext::FEvaluationNotifi
 		if (RootMotionProvider && !WarpedTrajectory.IsEmpty())
 		{
 			float Frame = (CurrentTime - ActualStartTime) / AnimNextAlignment::FrameTime;
+			FTransform WorldTransform;
+			AnimNextAlignment::GetTransformForFrame(Frame, WarpedTrajectory, WorldTransform);
 
-			int32 LowerFrame = FMath::Clamp(floor(Frame), 0, WarpedTrajectory.Num() -1);
-			int32 UpperFrame = FMath::Min(WarpedTrajectory.Num() -1 , LowerFrame + 1);
-			float Alpha = Frame - LowerFrame;
-		
-
-			FTransform WorldTransform = WarpedTrajectory[LowerFrame];
-			WorldTransform.BlendWith(WarpedTrajectory[UpperFrame], Alpha);
-
-			RootMotionProvider->OverrideRootMotion(WorldTransform.GetRelativeTransform(TraitInstanceData.RootBoneTransform), Keyframe->Get()->Attributes);
+			if(AlignmentNotify->UpdateMode == EAlignmentUpdateMode::World)
+			{
+				RootMotionProvider->OverrideRootMotion(WorldTransform.GetRelativeTransform(TraitInstanceData.RootBoneTransform), Keyframe->Get()->Attributes);
+			}
+			else // relative mode
+			{
+				FTransform PrevTransform;
+				AnimNextAlignment::GetTransformForFrame(PreviousFrame, WarpedTrajectory, PrevTransform);
+				RootMotionProvider->OverrideRootMotion(WorldTransform.GetRelativeTransform(PrevTransform), Keyframe->Get()->Attributes);
+				PreviousFrame = Frame;
+			}
 		
 			// unwarped trajectory relative to starting transform
 			FTransform PreviousTransform = AnimTrajectoryData.Trajectory[0] * StartingRootTransform;
diff --git a/Engine/Plugins/Experimental/EvaluationNotifies/Source/Runtime/Private/AnimNext/AnimNextAlignment.h b/Engine/Plugins/Experimental/EvaluationNotifies/Source/Runtime/Private/AnimNext/AnimNextAlignment.h
--- a/Engine/Plugins/Experimental/EvaluationNotifies/Source/Runtime/Private/AnimNext/AnimNextAlignment.h
+++ b/Engine/Plugins/Experimental/EvaluationNotifies/Source/Runtime/Private/AnimNext/AnimNextAlignment.h
@@ -36,6 +36,7 @@ struct FEvaluationNotify_AlignmentInstance : public FEvaluationNotify_BaseInstan
 	FBoneReference AlignBone;
 	bool bFirstFrame = true;
 	float ActualStartTime = 0;
+	float PreviousFrame = 0;
 
 	FTransform StartingRootTransform;
 	FTransform TargetTransform;
diff --git a/Engine/Plugins/Experimental/EvaluationNotifies/Source/Runtime/Private/EvaluationNotifies/AnimNotifyState_Alignment.cpp b/Engine/Plugins/Experimental/EvaluationNotifies/Source/Runtime/Private/EvaluationNotifies/AnimNotifyState_Alignment.cpp
--- a/Engine/Plugins/Experimental/EvaluationNotifies/Source/Runtime/Private/EvaluationNotifies/AnimNotifyState_Alignment.cpp
+++ b/Engine/Plugins/Experimental/EvaluationNotifies/Source/Runtime/Private/EvaluationNotifies/AnimNotifyState_Alignment.cpp
@@ -25,6 +25,16 @@ namespace
 			return AnimationAsset->ExtractRootMotion(Context);
 		}
 	}
+	
+	void GetTransformForFrame(float Frame, const TArray<FTransform>& Trajectory, FTransform& OutTransform)
+	{
+		int32 LowerFrame = FMath::Clamp(floor(Frame), 0, Trajectory.Num() -1);
+		int32 UpperFrame = FMath::Min(Trajectory.Num() -1 , LowerFrame + 1);
+		float Alpha = Frame - LowerFrame;
+
+		OutTransform = Trajectory[LowerFrame];
+		OutTransform.BlendWith(Trajectory[UpperFrame], Alpha);
+	}
 }
 
 void FAlignmentNotifyInstance::Start(const UAnimSequenceBase* AnimationAsset)
@@ -304,7 +314,17 @@ void FAlignmentNotifyInstance::Update(const UAnimSequenceBase* AnimationAsset, f
 		FTransform WorldTransform = WarpedTrajectory[LowerFrame];
 		WorldTransform.BlendWith(WarpedTrajectory[UpperFrame], Alpha);
 
-		RootMotionProvider->OverrideRootMotion(WorldTransform.GetRelativeTransform(RootBoneTransform), Output.CustomAttributes);
+		if(AlignmentNotify->UpdateMode == EAlignmentUpdateMode::World)
+		{
+			RootMotionProvider->OverrideRootMotion(WorldTransform.GetRelativeTransform(RootBoneTransform), Output.CustomAttributes);
+		}
+		else // relative mode
+		{
+			FTransform PrevTransform;
+			GetTransformForFrame(PreviousFrame, WarpedTrajectory, PrevTransform);
+			RootMotionProvider->OverrideRootMotion(WorldTransform.GetRelativeTransform(PrevTransform), Output.CustomAttributes);
+			PreviousFrame = Frame;
+		}
 		
 		// unwarped trajectory relative to starting transform
 		FTransform PreviousTransform = AnimTrajectoryData.Trajectory[0] * StartingRootTransform;
diff --git a/Engine/Plugins/Experimental/EvaluationNotifies/Source/Runtime/Public/EvaluationNotifies/AnimNotifyState_Alignment.h b/Engine/Plugins/Experimental/EvaluationNotifies/Source/Runtime/Public/EvaluationNotifies/AnimNotifyState_Alignment.h
--- a/Engine/Plugins/Experimental/EvaluationNotifies/Source/Runtime/Public/EvaluationNotifies/AnimNotifyState_Alignment.h
+++ b/Engine/Plugins/Experimental/EvaluationNotifies/Source/Runtime/Public/EvaluationNotifies/AnimNotifyState_Alignment.h
@@ -66,6 +66,15 @@ struct AnimTrajectoryData
 	TArray<float> RotationCurve;
 };
 
+UENUM()
+enum class EAlignmentUpdateMode
+{
+	// each frame, apply root motion to get to the expected part of the alignment path in wold space
+	World,
+	// each frame, apply the root motion from the alignment path relative to the previous frame
+	Relative,
+};
+
 UCLASS(BlueprintType, DisplayName="Alignment")
 class UNotifyState_Alignment : public UAnimNotifyState
 {
@@ -94,6 +103,9 @@ class UNotifyState_Alignment : public UAnimNotifyState
 	// // a Dynamic Target might move and so warping needs to be updated every frame
 	// UPROPERTY(EditAnywhere, Category = TargetTransform)
 	// bool bDynamicTarget;
+	
+	UPROPERTY(EditAnywhere, Category = TargetTransform)
+	EAlignmentUpdateMode UpdateMode;
 
 	// Steering will rotate the character based on the difference between the animated movement direction and the warped movement direction, to keep the character facing their direction of movement
 	UPROPERTY(EditAnywhere, Category = Steering)
@@ -117,6 +129,7 @@ struct FAlignmentNotifyInstance : public FEvaluationNotifyInstance
 	FBoneReference AlignBone;
 	bool bFirstFrame = true;
 	float ActualStartTime = 0;
+	float PreviousFrame = 0;
 
 	FTransform StartingRootTransform;
 	FTransform TargetTransform;
