[AnimNext] exposed HostObject to notifies

diff --git a/Engine/Plugins/Experimental/EvaluationNotifies/Source/Runtime/Private/AnimNext/AnimNextAlignment.cpp b/Engine/Plugins/Experimental/EvaluationNotifies/Source/Runtime/Private/AnimNext/AnimNextAlignment.cpp
--- a/Engine/Plugins/Experimental/EvaluationNotifies/Source/Runtime/Private/AnimNext/AnimNextAlignment.cpp
+++ b/Engine/Plugins/Experimental/EvaluationNotifies/Source/Runtime/Private/AnimNext/AnimNextAlignment.cpp
@@ -10,8 +10,6 @@
 #include "EvaluationNotifies/AnimNotifyState_Alignment.h"
 #include "EvaluationVM/EvaluationVM.h"
 #include "EvaluationVM/KeyframeState.h"
-#include "Graph/AnimNextGraphInstance.h"
-#include "TraitCore/NodeInstance.h"
 #include "VisualLogger/VisualLogger.h"
 
 void FEvaluationNotify_AlignmentInstance::Start()
@@ -58,9 +56,7 @@ namespace AnimNextAlignment
 
 float FEvaluationNotify_AlignmentInstance::GetWeight(float Time, const FAlignmentWarpCurve& WarpCurve) const 
 {
-	UNotifyState_Alignment* AlignmentNotify = Cast<UNotifyState_Alignment>(AnimNotify);
 	float Weight;
-	
 
 	if (WarpCurve.CurveType == EAlignmentWeightCurveType::FromRootMotionTranslation && !AnimTrajectoryData.TranslationCurve.IsEmpty())
 	{
@@ -141,7 +137,7 @@ void FEvaluationNotify_AlignmentInstance::Update(UE::AnimNext::FEvaluationNotifi
 			if (TraitInstanceData.DataInterface->GetVariable(AlignmentNotify->TransformName, TargetTransform) == EPropertyBagResult::Success)
 			{
 				static const FBox UnitBox(FVector(-10, -10, -10), FVector(10, 10, 10));
-				UE_VLOG_OBOX(TraitInstanceData.DebugObject, "Alignment", Display, UnitBox, TargetTransform.ToMatrixWithScale(), FColor::Blue, TEXT(""));
+				UE_VLOG_OBOX(TraitInstanceData.HostObject, "Alignment", Display, UnitBox, TargetTransform.ToMatrixWithScale(), FColor::Blue, TEXT(""));
 			
 				// get alignment transform
 				const float PredictionDelta = EndTime - CurrentTime;
@@ -150,7 +146,7 @@ void FEvaluationNotify_AlignmentInstance::Update(UE::AnimNext::FEvaluationNotifi
 				if (AlignBone.HasValidSetup())
 				{
 					// get alignment bone relative to predicted root motion end point, and remove that as an offset to the alignment target
-					// Alignment bone is expected to be at the alignment end point at the beginning of the Notify window  (we don't have a good way to prediuct the component space transform of it in the future)
+					// Alignment bone is expected to be at the alignment end point at the beginning of the Notify window  (we don't have a good way to predict the component space transform of it in the future)
 				 	
 					FTransform PredictedRootMotion = TraitInstanceData.OnExtractRootMotionAttribute.Execute(CurrentTime, PredictionDelta, false);
 
@@ -174,7 +170,6 @@ void FEvaluationNotify_AlignmentInstance::Update(UE::AnimNext::FEvaluationNotifi
 			
 				FTransform PredictedTransform;
 				float PredictionTime = CurrentTime;
-				float TargetTime = EndTime;
 
 				for(int i=0;i<NumFrames; i++)
 				{
@@ -331,7 +326,7 @@ void FEvaluationNotify_AlignmentInstance::Update(UE::AnimNext::FEvaluationNotifi
 			for(int i=1;i<AnimTrajectoryData.Trajectory.Num(); i++)
 			{
 				FTransform TransformFromRoot = AnimTrajectoryData.Trajectory[i] * StartingRootTransform;
-				UE_VLOG_SEGMENT(TraitInstanceData.DebugObject, "Alignment", Display, PreviousTransform.GetLocation(), TransformFromRoot.GetLocation(), i%2 == 0 ? FColor::Yellow : FColor::Red, TEXT(""));
+				UE_VLOG_SEGMENT(TraitInstanceData.HostObject, "Alignment", Display, PreviousTransform.GetLocation(), TransformFromRoot.GetLocation(), i%2 == 0 ? FColor::Yellow : FColor::Red, TEXT(""));
 				PreviousTransform = TransformFromRoot;
 			}
 		
@@ -341,18 +336,18 @@ void FEvaluationNotify_AlignmentInstance::Update(UE::AnimNext::FEvaluationNotifi
 			for(int i=1;i<AnimTrajectoryData.Trajectory.Num(); i++)
 			{
 				FTransform TransformFromTarget = AnimTrajectoryData.Trajectory[i] * InverseLastFrame * TargetTransform;
-				UE_VLOG_SEGMENT(TraitInstanceData.DebugObject, "Alignment", Display, PreviousTransform.GetLocation(), TransformFromTarget.GetLocation(), i%2 == 0 ? FColor::Yellow : FColor::Red, TEXT(""));
+				UE_VLOG_SEGMENT(TraitInstanceData.HostObject, "Alignment", Display, PreviousTransform.GetLocation(), TransformFromTarget.GetLocation(), i%2 == 0 ? FColor::Yellow : FColor::Red, TEXT(""));
 				PreviousTransform = TransformFromTarget;
 			}
 		
 			// the warped trajectory
 			for(int i=1; i<WarpedTrajectory.Num(); i++)
 			{
-				UE_VLOG_SEGMENT(TraitInstanceData.DebugObject, "Alignment", Display, WarpedTrajectory[i-1].GetLocation(), WarpedTrajectory[i].GetLocation(), i%2 == 0 ? FColor::Green : FColor::Blue, TEXT(""));
+				UE_VLOG_SEGMENT(TraitInstanceData.HostObject, "Alignment", Display, WarpedTrajectory[i-1].GetLocation(), WarpedTrajectory[i].GetLocation(), i%2 == 0 ? FColor::Green : FColor::Blue, TEXT(""));
 			}
 
 			// a dot representing our current position on the trajectory 
-			UE_VLOG_SPHERE(TraitInstanceData.DebugObject, "Alignment", Display, WorldTransform.GetLocation(),  1, FColor::Red, TEXT(""));
+			UE_VLOG_SPHERE(TraitInstanceData.HostObject, "Alignment", Display, WorldTransform.GetLocation(),  1, FColor::Red, TEXT(""));
 		}
 	}
 }
diff --git a/Engine/Plugins/Experimental/EvaluationNotifies/Source/Runtime/Private/AnimNext/EvaluationNotifiesTrait.cpp b/Engine/Plugins/Experimental/EvaluationNotifies/Source/Runtime/Private/AnimNext/EvaluationNotifiesTrait.cpp
--- a/Engine/Plugins/Experimental/EvaluationNotifies/Source/Runtime/Private/AnimNext/EvaluationNotifiesTrait.cpp
+++ b/Engine/Plugins/Experimental/EvaluationNotifies/Source/Runtime/Private/AnimNext/EvaluationNotifiesTrait.cpp
@@ -120,9 +120,7 @@ namespace UE::AnimNext
 		FInstanceData* InstanceData = Binding.GetInstanceData<FInstanceData>();
 		check(InstanceData);
 
-#if ENABLE_ANIM_DEBUG
-		InstanceData->DebugObject = Context.GetHostObject();
-#endif
+		InstanceData->HostObject = Context.GetHostObject();
 
 		// Get Root motion provider
 		{
diff --git a/Engine/Plugins/Experimental/EvaluationNotifies/Source/Runtime/Private/AnimNext/EvaluationNotifiesTrait.h b/Engine/Plugins/Experimental/EvaluationNotifies/Source/Runtime/Private/AnimNext/EvaluationNotifiesTrait.h
--- a/Engine/Plugins/Experimental/EvaluationNotifies/Source/Runtime/Private/AnimNext/EvaluationNotifiesTrait.h
+++ b/Engine/Plugins/Experimental/EvaluationNotifies/Source/Runtime/Private/AnimNext/EvaluationNotifiesTrait.h
@@ -56,18 +56,15 @@ struct EVALUATIONNOTIFIESRUNTIME_API FEvaluationNotifiesTrait : FAdditiveTrait,
 		/** Last root bone transform sampled */
 		FTransform RootBoneTransform = FTransform::Identity;
 
-		// the evaluation notifies extracted from the current anim sequence
+		/** the evaluation notifies extracted from the current anim sequence */
 		TArray<FInstancedStruct> EvaluationNotifies;
 
 		const FAnimNextDataInterfaceInstance* DataInterface = nullptr;
 
-#if ENABLE_ANIM_DEBUG
-		// an object for attaching debug draws to
-		const UObject* DebugObject;
-#endif
+		/** Owner object the trait instance is associated to */
+		const UObject* HostObject;
 	};
-	
-	
+
 	// ITimelinePlayer impl
 	virtual void AdvanceBy(FExecutionContext& Context, const TTraitBinding<ITimelinePlayer>& Binding, float DeltaTime) const override;
 	
@@ -79,7 +76,7 @@ struct EVALUATIONNOTIFIESRUNTIME_API FEvaluationNotifiesTrait : FAdditiveTrait,
 
 	static void RegisterEvaluationHandler(UClass* NotifyType, UScriptStruct* Handler);
 	static void UnregisterEvaluationHandler(UClass* NotifyType);
-	
+
 private:
 	static TMap<UClass*, UScriptStruct*> NotifyEvaluationHandlerMap;
 };
