[AnimNext] fixed missing call to End() for aborted EvaluationNotifiesTrait
Also moved EvaluationNotifies to a public header

diff --git a/Engine/Plugins/Experimental/EvaluationNotifies/Source/Runtime/Private/AnimNext/AnimNextAlignment.h b/Engine/Plugins/Experimental/EvaluationNotifies/Source/Runtime/Private/AnimNext/AnimNextAlignment.h
--- a/Engine/Plugins/Experimental/EvaluationNotifies/Source/Runtime/Private/AnimNext/AnimNextAlignment.h
+++ b/Engine/Plugins/Experimental/EvaluationNotifies/Source/Runtime/Private/AnimNext/AnimNextAlignment.h
@@ -2,19 +2,11 @@
 
 #pragma once
 
-#include "CoreMinimal.h"
-#include "EvaluationNotifiesTrait.h"
-#include "DataInterface/AnimNextDataInterfaceInstance.h"
-#include "EvaluationNotifies/AnimNotifyState_Alignment.h"
-#include "Kismet/KismetMathLibrary.h"
-#include "TraitCore/Trait.h"
-#include "TraitCore/TraitSharedData.h"
-#include "TraitInterfaces/IAttributeProvider.h"
-#include "TraitInterfaces/IEvaluate.h"
-#include "TraitInterfaces/ITimeline.h"
-#include "TraitInterfaces/IUpdate.h"
+#include "AnimNext/EvaluationNotifiesTrait.h"
 #include "AnimNextAlignment.generated.h"
 
+struct FAlignmentWarpCurve;
+
 struct AlignmentAnimTrajectoryData
 {
 	FTransform TargetTransform;
diff --git a/Engine/Plugins/Experimental/EvaluationNotifies/Source/Runtime/Private/AnimNext/EvaluationNotifiesTrait.cpp b/Engine/Plugins/Experimental/EvaluationNotifies/Source/Runtime/Private/AnimNext/EvaluationNotifiesTrait.cpp
--- a/Engine/Plugins/Experimental/EvaluationNotifies/Source/Runtime/Private/AnimNext/EvaluationNotifiesTrait.cpp
+++ b/Engine/Plugins/Experimental/EvaluationNotifies/Source/Runtime/Private/AnimNext/EvaluationNotifiesTrait.cpp
@@ -1,22 +1,11 @@
 // Copyright Epic Games, Inc. All Rights Reserved.
 
-#include "EvaluationNotifiesTrait.h"
+#include "AnimNext/EvaluationNotifiesTrait.h"
 
-#include "AnimNextAnimGraphSettings.h"
-// #include "AnimNextWarpingLog.h"
-#include <AHEasing/easing.h>
-
-#include "Animation/AnimRootMotionProvider.h"
-#include "Animation/AnimSequence.h"
-#include "EvaluationNotifies/AnimNotifyState_Alignment.h"
-#include "EvaluationVM/EvaluationVM.h"
-#include "Graph/AnimNextGraphInstance.h"
-#include "TraitCore/NodeInstance.h"
-#include "TraitInterfaces/IBlendStack.h"
-#include "TraitInterfaces/IGraphFactory.h"
+#include "Animation/AnimNotifies/AnimNotifyState.h"
 #include "TraitInterfaces/INotifySource.h"
-#include "TraitInterfaces/ITimelinePlayer.h"
 #include "TraitInterfaces/ITimeline.h"
+#include "TraitCore/NodeInstance.h"
 
 static const FName RootBoneTransformName = "RootBoneTransform";
 
@@ -74,7 +63,7 @@ namespace UE::AnimNext
 				{
 					if (UScriptStruct** HandlerType = NotifyEvaluationHandlerMap.Find(Notify->GetClass()))
 					{
-						if (FInstancedStruct* FoundInstance = InstanceData->EvaluationNotifies.FindByPredicate([NotifyEvent](const FInstancedStruct ExistingInstance)
+						if (FInstancedStruct* FoundInstance = InstanceData->EvaluationNotifies.FindByPredicate([NotifyEvent](const FStructView ExistingInstance)
 						{
 							// check if we already have an instance for this event
 							return ExistingInstance.Get<FEvaluationNotify_BaseInstance>().NotifyEvent == NotifyEvent;
@@ -140,17 +129,31 @@ namespace UE::AnimNext
 		Context.AppendTask(FAnimNextEvaluationNotifiesTask::Make(InstanceData, SharedData));
 	}
 
-
 	void FEvaluationNotifiesTrait::RegisterEvaluationHandler(UClass* NotifyType, UScriptStruct* HandlerType)
-    {
-    	NotifyEvaluationHandlerMap.Add(NotifyType, HandlerType);
-    }
-    
-    void FEvaluationNotifiesTrait::UnregisterEvaluationHandler(UClass* NotifyType)
-    {
-    	NotifyEvaluationHandlerMap.Remove(NotifyType);
-    }
-	
+	{
+		NotifyEvaluationHandlerMap.Add(NotifyType, HandlerType);
+	}
+
+	void FEvaluationNotifiesTrait::UnregisterEvaluationHandler(UClass* NotifyType)
+	{
+		NotifyEvaluationHandlerMap.Remove(NotifyType);
+	}
+
+	void FEvaluationNotifiesTrait::FInstanceData::Destruct(const FExecutionContext& Context, const FTraitBinding& Binding)
+	{
+		for (FStructView NotifyInstanceData : Binding.GetInstanceData<FInstanceData>()->EvaluationNotifies)
+		{
+			FEvaluationNotify_BaseInstance& NotifyInstance = NotifyInstanceData.Get<FEvaluationNotify_BaseInstance>();
+			if (NotifyInstance.bActive)
+			{
+				NotifyInstance.bActive = false;
+				NotifyInstance.End();
+			}
+		}
+
+		FTrait::FInstanceData::Destruct(Context, Binding);
+	}
+
 } // namespace UE::AnimNext
 
 
diff --git a/Engine/Plugins/Experimental/EvaluationNotifies/Source/Runtime/Private/AnimNext/EvaluationNotifiesTrait.h b/Engine/Plugins/Experimental/EvaluationNotifies/Source/Runtime/Private/AnimNext/EvaluationNotifiesTrait.h
deleted file mode 100644
diff --git a/Engine/Plugins/Experimental/EvaluationNotifies/Source/Runtime/Public/AnimNext/EvaluationNotifiesTrait.h b/Engine/Plugins/Experimental/EvaluationNotifies/Source/Runtime/Public/AnimNext/EvaluationNotifiesTrait.h
new file mode 100644
--- /dev/null
+++ b/Engine/Plugins/Experimental/EvaluationNotifies/Source/Runtime/Public/AnimNext/EvaluationNotifiesTrait.h
@@ -0,0 +1,120 @@
+// Copyright Epic Games, Inc. All Rights Reserved.
+
+#pragma once
+
+#include "CoreMinimal.h"
+#include "DataInterface/AnimNextDataInterfaceInstance.h"
+#include "Kismet/KismetMathLibrary.h"
+#include "TraitCore/Trait.h"
+#include "TraitCore/TraitSharedData.h"
+#include "TraitInterfaces/IAttributeProvider.h"
+#include "TraitInterfaces/IEvaluate.h"
+#include "TraitInterfaces/ITimelinePlayer.h"
+#include "TraitInterfaces/IUpdate.h"
+#include "EvaluationNotifiesTrait.generated.h"
+
+struct FEvaluationNotify_BaseInstance;
+
+/**
+ * Data needed to execute steering
+ * 
+ * Some steering data such as current anim asset / playback time is acquired via trait stack interfaces
+ */
+USTRUCT(meta = (DisplayName = "EvaluationNotifies"))
+struct FEvaluationNotifiesTraitSharedData : public FAnimNextTraitSharedData
+{
+	GENERATED_BODY()
+
+	// Latent pin support boilerplate
+	#define TRAIT_LATENT_PROPERTIES_ENUMERATOR(GeneratorMacro) \
+
+	GENERATE_TRAIT_LATENT_PROPERTIES(FEvaluationNotifiesTraitSharedData, TRAIT_LATENT_PROPERTIES_ENUMERATOR)
+	#undef TRAIT_LATENT_PROPERTIES_ENUMERATOR
+};
+
+namespace UE::AnimNext
+{
+/**
+ * this trait will run evaluation time code for notifies which have a registered handler
+ */
+struct EVALUATIONNOTIFIESRUNTIME_API FEvaluationNotifiesTrait : FAdditiveTrait, ITimelinePlayer, IUpdate, IEvaluate
+{
+	DECLARE_ANIM_TRAIT(FEvaluationNotifiesTrait, FAdditiveTrait)
+
+	using FSharedData = FEvaluationNotifiesTraitSharedData;
+
+	struct FInstanceData : FTrait::FInstanceData
+	{
+		void Destruct(const FExecutionContext& Context, const FTraitBinding& Binding);
+
+		/** Delta in seconds between updates, populated during PreUpdate */
+		float DeltaTime = 0.f;
+
+		/** Callback provided by attribute trait on stack to evaluate root motion at a later time */
+		FOnExtractRootMotionAttribute OnExtractRootMotionAttribute = FOnExtractRootMotionAttribute();
+
+		/** Last root bone transform sampled */
+		FTransform RootBoneTransform = FTransform::Identity;
+
+		/** the evaluation notifies extracted from the current anim sequence */
+		TArray<FInstancedStruct> EvaluationNotifies;
+
+		const FAnimNextDataInterfaceInstance* DataInterface = nullptr;
+
+		/** Owner object the trait instance is associated to */
+		const UObject* HostObject;
+	};
+
+	// ITimelinePlayer impl
+	virtual void AdvanceBy(FExecutionContext& Context, const TTraitBinding<ITimelinePlayer>& Binding, float DeltaTime) const override;
+	
+	// IUpdate impl 
+	virtual void PreUpdate(FUpdateTraversalContext& Context, const TTraitBinding<IUpdate>& Binding, const FTraitUpdateState& TraitState) const override;
+
+	// IEvaluate impl
+	virtual void PostEvaluate(FEvaluateTraversalContext& Context, const TTraitBinding<IEvaluate>& Binding) const override;
+
+	static void RegisterEvaluationHandler(UClass* NotifyType, UScriptStruct* Handler);
+	static void UnregisterEvaluationHandler(UClass* NotifyType);
+
+private:
+	static TMap<UClass*, UScriptStruct*> NotifyEvaluationHandlerMap;
+};
+
+} // namespace UE::AnimNext
+
+/** Task to run EvaluationNotifies on VM */
+USTRUCT()
+struct FAnimNextEvaluationNotifiesTask : public FAnimNextEvaluationTask
+{
+	GENERATED_BODY()
+
+	DECLARE_ANIM_EVALUATION_TASK(FAnimNextEvaluationNotifiesTask)
+
+	static FAnimNextEvaluationNotifiesTask Make(UE::AnimNext::FEvaluationNotifiesTrait::FInstanceData* InstanceData, const UE::AnimNext::FEvaluationNotifiesTrait::FSharedData* SharedData);
+
+	// Task entry point
+	virtual void Execute(UE::AnimNext::FEvaluationVM& VM) const override;
+
+	UE::AnimNext::FEvaluationNotifiesTrait::FInstanceData* InstanceData = nullptr;
+	const UE::AnimNext::FEvaluationNotifiesTrait::FSharedData* SharedData = nullptr;
+};
+
+USTRUCT()
+struct FEvaluationNotify_BaseInstance
+{
+	GENERATED_BODY()
+	virtual void Start() {}
+	virtual void Update(UE::AnimNext::FEvaluationNotifiesTrait::FInstanceData& InstanceData, UE::AnimNext::FEvaluationVM& VM) {}
+	virtual void End() {}
+	
+	virtual ~FEvaluationNotify_BaseInstance() {};
+
+	const FAnimNotifyEvent* NotifyEvent;
+	TObjectPtr<UAnimNotifyState> AnimNotify = nullptr;
+	float StartTime = 0.f;
+	float EndTime = 0.f;
+	
+	float CurrentTime = 0.f;
+	bool bActive = false;
+};
\ No newline at end of file
