Temp fix crash alignment when delta time is less than a frame

#tests PIE

diff --git a/Engine/Plugins/Experimental/EvaluationNotifies/Source/Runtime/Private/AnimNext/AnimNextAlignment.cpp b/Engine/Plugins/Experimental/EvaluationNotifies/Source/Runtime/Private/AnimNext/AnimNextAlignment.cpp
--- a/Engine/Plugins/Experimental/EvaluationNotifies/Source/Runtime/Private/AnimNext/AnimNextAlignment.cpp
+++ b/Engine/Plugins/Experimental/EvaluationNotifies/Source/Runtime/Private/AnimNext/AnimNextAlignment.cpp
@@ -162,6 +162,13 @@ void FEvaluationNotify_AlignmentInstance::Update(UE::AnimNext::FEvaluationNotifi
 			
 				int NumFrames = (EndTime - CurrentTime) / AnimNextAlignment::FrameTime;
 
+				// @TODO: Temp to prevent crash. Should do a ceil on above 'EndTime - CurrentTime'
+				// Also consider adding 1 when when needed to ensure last bit of sub-frame root motion gets consumed.
+				if (NumFrames <= 0)
+				{
+					return;
+				}
+
 				AnimTrajectoryData.Trajectory.SetNum(NumFrames);
 				AnimTrajectoryData.TranslationCurve.SetNum(NumFrames);
 				AnimTrajectoryData.RotationCurve.SetNum(NumFrames);
