AlignToGround feature for traversals

diff --git a/Engine/Plugins/Experimental/EvaluationNotifies/Source/Runtime/Private/AnimNext/AnimNextAlignment.cpp b/Engine/Plugins/Experimental/EvaluationNotifies/Source/Runtime/Private/AnimNext/AnimNextAlignment.cpp
--- a/Engine/Plugins/Experimental/EvaluationNotifies/Source/Runtime/Private/AnimNext/AnimNextAlignment.cpp
+++ b/Engine/Plugins/Experimental/EvaluationNotifies/Source/Runtime/Private/AnimNext/AnimNextAlignment.cpp
@@ -14,7 +14,7 @@
 
 void FEvaluationNotify_AlignmentInstance::Start()
 {
-	const UNotifyState_Alignment* AlignmentNotify = Cast<UNotifyState_Alignment>(AnimNotify);
+	const UNotifyState_AlignmentBase* AlignmentNotify = CastChecked<UNotifyState_AlignmentBase>(AnimNotify);
 	AlignBone = AlignmentNotify->AlignBone;
 	bFirstFrame = true;
 }
@@ -103,6 +103,58 @@ float FEvaluationNotify_AlignmentInstance::GetWeight(float Time, const FAlignmen
 	return Weight;
 }
 
+bool FEvaluationNotify_AlignToGroundInstance::GetTargetTransform(UE::AnimNext::FEvaluationNotifiesTrait::FInstanceData& TraitInstanceData, FTransform& Transform)
+{
+	const UNotifyState_AlignToGround* AlignToGroundNotify = CastChecked<UNotifyState_AlignToGround>(AnimNotify);
+	const FCollisionShape CollisionShape = FCollisionShape::MakeSphere(AlignToGroundNotify->TraceRadius);
+
+	FTransform StartTransform = TraitInstanceData.RootBoneTransform;
+	
+	const float PredictionDelta = FMath::Max(EndTime - CurrentTime, 0);
+	FTransform PredictedRootMotion = TraitInstanceData.OnExtractRootMotionAttribute.Execute(CurrentTime, PredictionDelta, false);
+	FTransform EndTransform = PredictedRootMotion * StartTransform;
+
+	const FVector TraceDirectionWS = FVector::UpVector;
+	const FVector TraceStart = EndTransform.GetLocation() + (TraceDirectionWS * AlignToGroundNotify->TraceStartOffset);
+	const FVector TraceEnd = EndTransform.GetLocation() + (AlignToGroundNotify->TraceEndOffset * TraceDirectionWS);
+
+	FCollisionQueryParams QueryParams;
+	// Ignore self and all attached components
+	// QueryParams.AddIgnoredActor(Context.OwningActor);
+
+	const ECollisionChannel CollisionChannel = UEngineTypes::ConvertToCollisionChannel(AlignToGroundNotify->TraceChannel);
+
+	if (const UObject* HostObject = TraitInstanceData.HostObject)
+	{
+		UWorld* World = HostObject->GetWorld();
+		check(World);
+
+		FHitResult HitResult;
+		const bool bHit = World->SweepSingleByChannel(
+		HitResult, TraceStart, TraceEnd, FQuat::Identity, CollisionChannel, CollisionShape, QueryParams);
+
+		UE_VLOG_CAPSULE(TraitInstanceData.HostObject, "AlignToGround", Display,
+					TraceEnd,
+					(AlignToGroundNotify->TraceStartOffset - AlignToGroundNotify->TraceEndOffset) * 0.5f,
+					AlignToGroundNotify->TraceRadius,
+					FQuat::Identity,
+					FColor::Green,
+					TEXT(""));
+
+		if (bHit)
+		{
+			Transform = EndTransform;
+			Transform.SetLocation(HitResult.ImpactPoint);
+		
+			static const FBox UnitBox(FVector(-10, -10, -10), FVector(10, 10, 10));
+			UE_VLOG_OBOX(TraitInstanceData.HostObject, "AlignToGround", Display, UnitBox, Transform.ToMatrixWithScale(), FColor::Red, TEXT(""));
+			return true;
+		}
+	}
+
+	return false;
+}
+
 FTransform GetModelSpaceTransform(const UE::AnimNext::FLODPoseStack& Pose, FBoneIndexType SkeletonBoneIndex)
 {
 	const TArrayView<const FBoneIndexType> SkeletonToPoseIndexMap = Pose.GetSkeletonBoneIndexToLODBoneIndexMap();
@@ -120,9 +172,15 @@ FTransform GetModelSpaceTransform(const UE::AnimNext::FLODPoseStack& Pose, FBone
 	return BoneTransform;
 }
 
+bool FEvaluationNotify_AlignmentInstance::GetTargetTransform(UE::AnimNext::FEvaluationNotifiesTrait::FInstanceData& TraitInstanceData, FTransform& Transform)
+{
+	const UNotifyState_Alignment* AlignmentNotify = CastChecked<UNotifyState_Alignment>(AnimNotify);
+	return TraitInstanceData.DataInterface->GetVariable(AlignmentNotify->TransformName, Transform) == EPropertyBagResult::Success;
+}
+
 void FEvaluationNotify_AlignmentInstance::Update(UE::AnimNext::FEvaluationNotifiesTrait::FInstanceData& TraitInstanceData, UE::AnimNext::FEvaluationVM& VM)
 {
-	UNotifyState_Alignment* AlignmentNotify = Cast<UNotifyState_Alignment>(AnimNotify);
+	UNotifyState_AlignmentBase* AlignmentNotify = CastChecked<UNotifyState_AlignmentBase>(AnimNotify);
 	
 	if (const TUniquePtr<UE::AnimNext::FKeyframeState>* Keyframe = VM.PeekValue<TUniquePtr<UE::AnimNext::FKeyframeState>>(UE::AnimNext::KEYFRAME_STACK_NAME, 0))
 	{
@@ -134,14 +192,13 @@ void FEvaluationNotify_AlignmentInstance::Update(UE::AnimNext::FEvaluationNotifi
 			const UE::AnimNext::FLODPoseStack& Pose = Keyframe->Get()->Pose;
 			AlignBone.Initialize(Pose.GetSkeletonAsset());
 
-			if (TraitInstanceData.DataInterface->GetVariable(AlignmentNotify->TransformName, TargetTransform) == EPropertyBagResult::Success)
+			if (GetTargetTransform(TraitInstanceData, TargetTransform))
 			{
 				static const FBox UnitBox(FVector(-10, -10, -10), FVector(10, 10, 10));
 				UE_VLOG_OBOX(TraitInstanceData.HostObject, "Alignment", Display, UnitBox, TargetTransform.ToMatrixWithScale(), FColor::Blue, TEXT(""));
 			
 				// get alignment transform
 				const float PredictionDelta = EndTime - CurrentTime;
-			
 
 				if (AlignBone.HasValidSetup())
 				{
@@ -160,10 +217,8 @@ void FEvaluationNotify_AlignmentInstance::Update(UE::AnimNext::FEvaluationNotifi
 			
 				// extract root motion trajectory todo: this should be cached and reused
 			
-				int NumFrames = (EndTime - CurrentTime) / AnimNextAlignment::FrameTime;
+				int NumFrames = 1 + (EndTime - CurrentTime) / AnimNextAlignment::FrameTime;
 
-				// @TODO: Temp to prevent crash. Should do a ceil on above 'EndTime - CurrentTime'
-				// Also consider adding 1 when when needed to ensure last bit of sub-frame root motion gets consumed.
 				if (NumFrames <= 0)
 				{
 					return;
@@ -248,7 +303,7 @@ void FEvaluationNotify_AlignmentInstance::Update(UE::AnimNext::FEvaluationNotifi
 				// Translation warping + steering
 				for(int i=0;i<NumFrames; i++)
 				{
-					float Weight = GetWeight(ActualStartTime + AnimNextAlignment::FrameTime*i, AlignmentNotify->TranslationWarpingCurve);
+					float Weight = GetWeight(ActualStartTime + AnimNextAlignment::FrameTime * (i+1), AlignmentNotify->TranslationWarpingCurve);
 
 					FTransform TransformFromRoot = AnimTrajectoryData.Trajectory[i] * TraitInstanceData.RootBoneTransform;
 					FTransform TransformFromTarget = AnimTrajectoryData.Trajectory[i] * InverseLastFrame * TargetTransform;
diff --git a/Engine/Plugins/Experimental/EvaluationNotifies/Source/Runtime/Private/AnimNext/AnimNextAlignment.h b/Engine/Plugins/Experimental/EvaluationNotifies/Source/Runtime/Private/AnimNext/AnimNextAlignment.h
--- a/Engine/Plugins/Experimental/EvaluationNotifies/Source/Runtime/Private/AnimNext/AnimNextAlignment.h
+++ b/Engine/Plugins/Experimental/EvaluationNotifies/Source/Runtime/Private/AnimNext/AnimNextAlignment.h
@@ -21,6 +21,7 @@ struct FEvaluationNotify_AlignmentInstance : public FEvaluationNotify_BaseInstan
 {
 	GENERATED_BODY()
 	virtual void Start() override;
+	virtual bool GetTargetTransform(UE::AnimNext::FEvaluationNotifiesTrait::FInstanceData& TraitInstanceData, FTransform& TargetTranform);
    	virtual void Update(UE::AnimNext::FEvaluationNotifiesTrait::FInstanceData& InstanceData, UE::AnimNext::FEvaluationVM& VM) override;
 	
 	float GetWeight(float Time, const FAlignmentWarpCurve& WarpCurve) const; 
@@ -38,4 +39,11 @@ struct FEvaluationNotify_AlignmentInstance : public FEvaluationNotify_BaseInstan
 
 	TArray<FTransform> WarpedTrajectory;
 	AlignmentAnimTrajectoryData AnimTrajectoryData;
+};
+
+USTRUCT()
+struct FEvaluationNotify_AlignToGroundInstance : public FEvaluationNotify_AlignmentInstance
+{
+	GENERATED_BODY()
+	virtual bool GetTargetTransform(UE::AnimNext::FEvaluationNotifiesTrait::FInstanceData& TraitInstanceData, FTransform& TargetTranform);
 };
\ No newline at end of file
diff --git a/Engine/Plugins/Experimental/EvaluationNotifies/Source/Runtime/Private/EvaluationNotifiesRuntimeModule.cpp b/Engine/Plugins/Experimental/EvaluationNotifies/Source/Runtime/Private/EvaluationNotifiesRuntimeModule.cpp
--- a/Engine/Plugins/Experimental/EvaluationNotifies/Source/Runtime/Private/EvaluationNotifiesRuntimeModule.cpp
+++ b/Engine/Plugins/Experimental/EvaluationNotifies/Source/Runtime/Private/EvaluationNotifiesRuntimeModule.cpp
@@ -38,6 +38,7 @@ namespace UE
 			FAnimNode_EvaluationNotifies::RegisterEvaluationHandler(UNotifyState_TwoBoneIK::StaticClass(), FTwoBoneIKNotifyInstance::StaticStruct());
 
 			AnimNext::FEvaluationNotifiesTrait::RegisterEvaluationHandler(UNotifyState_Alignment::StaticClass(), FEvaluationNotify_AlignmentInstance::StaticStruct());
+			AnimNext::FEvaluationNotifiesTrait::RegisterEvaluationHandler(UNotifyState_AlignToGround::StaticClass(), FEvaluationNotify_AlignToGroundInstance::StaticStruct());
 		}
 
 		void FModule::ShutdownModule()
diff --git a/Engine/Plugins/Experimental/EvaluationNotifies/Source/Runtime/Public/EvaluationNotifies/AnimNotifyState_Alignment.h b/Engine/Plugins/Experimental/EvaluationNotifies/Source/Runtime/Public/EvaluationNotifies/AnimNotifyState_Alignment.h
--- a/Engine/Plugins/Experimental/EvaluationNotifies/Source/Runtime/Public/EvaluationNotifies/AnimNotifyState_Alignment.h
+++ b/Engine/Plugins/Experimental/EvaluationNotifies/Source/Runtime/Public/EvaluationNotifies/AnimNotifyState_Alignment.h
@@ -75,8 +75,8 @@ enum class EAlignmentUpdateMode
 	Relative,
 };
 
-UCLASS(BlueprintType, DisplayName="Alignment")
-class UNotifyState_Alignment : public UAnimNotifyState
+UCLASS()
+class UNotifyState_AlignmentBase : public UAnimNotifyState
 {
 	GENERATED_BODY()
 	
@@ -88,22 +88,24 @@ class UNotifyState_Alignment : public UAnimNotifyState
 	UPROPERTY(EditAnywhere, Category = Blending)
 	FAlignmentWarpCurve RotationWarpingCurve = { };
 	
-	// Named transform to align to
-    UPROPERTY(EditAnywhere, Category = TargetTransform)
-    FName TransformName;
 
 	// Offset from Root (or AlignBone) to Align to the target transform
 	UPROPERTY(EditAnywhere, Category = TargetTransform)
 	FTransform AlignOffset;
 
+#if WITH_EDITOR
+	// For automating setting the AlignOffset at the end of the NotifyState
+    UFUNCTION(CallInEditor, BlueprintCallable, Category = "TechAnimUtils")
+    void SetAlignOffset(FTransform NewTransform)
+    {
+		AlignOffset = NewTransform;
+    }
+#endif
+
 	// Optional bone to Align to target transform
 	UPROPERTY(EditAnywhere, Category = TargetTransform)
 	FBoneReference AlignBone;
 		
-	// // a Dynamic Target might move and so warping needs to be updated every frame
-	// UPROPERTY(EditAnywhere, Category = TargetTransform)
-	// bool bDynamicTarget;
-	
 	UPROPERTY(EditAnywhere, Category = TargetTransform)
 	EAlignmentUpdateMode UpdateMode;
 
@@ -115,6 +117,25 @@ class UNotifyState_Alignment : public UAnimNotifyState
 	FAlignmentSteeringSettings SteeringSettings;
 };
 
+UCLASS(BlueprintType, DisplayName="Alignment")
+class UNotifyState_Alignment : public UNotifyState_AlignmentBase
+{
+	GENERATED_BODY()
+	
+	public:
+	// Named transform to align to
+    UPROPERTY(EditAnywhere, Category = TargetTransform)
+    FName TransformName;
+
+#if WITH_EDITOR
+	UFUNCTION(CallInEditor, BlueprintCallable, Category = "TechAnimUtils")
+	void SetTransformName(FName NewName)
+	{
+		TransformName = NewName;
+	}
+#endif
+};
+
 USTRUCT()
 struct FAlignmentNotifyInstance : public FEvaluationNotifyInstance
 {
@@ -139,4 +160,24 @@ struct FAlignmentNotifyInstance : public FEvaluationNotifyInstance
 
 	TArray<FTransform> WarpedTrajectory;
 	AnimTrajectoryData AnimTrajectoryData;
+};
+
+UCLASS(BlueprintType, DisplayName="AlignToGround")
+class UNotifyState_AlignToGround : public UNotifyState_AlignmentBase
+{
+	GENERATED_BODY()
+
+	public:
+	
+	UPROPERTY(EditAnywhere, Category = "Trace Settings")
+	float TraceRadius;
+	
+	UPROPERTY(EditAnywhere, Category = "Trace Settings")
+	float TraceStartOffset;
+	
+	UPROPERTY(EditAnywhere, Category = "Trace Settings")
+	float TraceEndOffset;
+	
+	UPROPERTY(EditAnywhere, Category = "Trace Settings")
+	TEnumAsByte<ETraceTypeQuery> TraceChannel = TraceTypeQuery1;
 };
\ No newline at end of file
