Add time scaling support to AlignToGround notify state, and some improvements for triggering evaluationnotify end callback

diff --git a/Engine/Plugins/Experimental/EvaluationNotifies/Source/Runtime/Private/AnimNext/AnimNextAlignment.cpp b/Engine/Plugins/Experimental/EvaluationNotifies/Source/Runtime/Private/AnimNext/AnimNextAlignment.cpp
--- a/Engine/Plugins/Experimental/EvaluationNotifies/Source/Runtime/Private/AnimNext/AnimNextAlignment.cpp
+++ b/Engine/Plugins/Experimental/EvaluationNotifies/Source/Runtime/Private/AnimNext/AnimNextAlignment.cpp
@@ -103,6 +103,14 @@ float FEvaluationNotify_AlignmentInstance::GetWeight(float Time, const FAlignmen
 	return Weight;
 }
 
+void FEvaluationNotify_AlignToGroundInstance::End(UE::AnimNext::FEvaluationNotifiesTrait::FInstanceData& TraitInstanceData)
+{
+	FEvaluationNotify_AlignmentInstance::End(TraitInstanceData);
+	
+	const UNotifyState_AlignToGround* AlignToGroundNotify = CastChecked<UNotifyState_AlignToGround>(AnimNotify);
+	TraitInstanceData.DataInterface->SetVariable(AlignToGroundNotify->PlaybackRateOutputVariable, 1.0);
+}
+
 bool FEvaluationNotify_AlignToGroundInstance::GetTargetTransform(UE::AnimNext::FEvaluationNotifiesTrait::FInstanceData& TraitInstanceData, FTransform& Transform)
 {
 	const UNotifyState_AlignToGround* AlignToGroundNotify = CastChecked<UNotifyState_AlignToGround>(AnimNotify);
@@ -148,6 +156,28 @@ bool FEvaluationNotify_AlignToGroundInstance::GetTargetTransform(UE::AnimNext::F
 		
 			static const FBox UnitBox(FVector(-10, -10, -10), FVector(10, 10, 10));
 			UE_VLOG_OBOX(TraitInstanceData.HostObject, "AlignToGround", Display, UnitBox, Transform.ToMatrixWithScale(), FColor::Red, TEXT(""));
+
+			if (!AlignToGroundNotify->PlaybackRateOutputVariable.IsNone())
+			{
+				float AnimatedFallDistance = StartTransform.GetTranslation().Z - EndTransform.GetTranslation().Z;
+				
+				if (AnimatedFallDistance > UE_KINDA_SMALL_NUMBER && PredictionDelta > UE_KINDA_SMALL_NUMBER)
+				{
+					// Math computes acceleration due to gravity as animated,
+					// and then uses that to determine a modified playback rate for different actual fall heights (such that at the animated height the modifier will be 1.0)
+					// - assuming starting from 0 vertical velocity (which is typically where you would want to start this notify, as you start falling)
+					
+					float AnimatedFallingAcceleration = 2.0f * AnimatedFallDistance / (PredictionDelta * PredictionDelta);
+				
+					float ActualFallDistance = StartTransform.GetTranslation().Z - HitResult.ImpactPoint.Z;
+					float ModifiedTime = sqrtf(2.0f * ActualFallDistance/AnimatedFallingAcceleration);
+				
+					float PlaybackRate = FMath::Clamp(PredictionDelta / ModifiedTime, AlignToGroundNotify->MinPlaybackRateModifier, AlignToGroundNotify->MaxPlaybackRateModifier);
+				
+					TraitInstanceData.DataInterface->SetVariable(AlignToGroundNotify->PlaybackRateOutputVariable, (double)PlaybackRate);
+				}
+			}
+			
 			return true;
 		}
 	}
diff --git a/Engine/Plugins/Experimental/EvaluationNotifies/Source/Runtime/Private/AnimNext/AnimNextAlignment.h b/Engine/Plugins/Experimental/EvaluationNotifies/Source/Runtime/Private/AnimNext/AnimNextAlignment.h
--- a/Engine/Plugins/Experimental/EvaluationNotifies/Source/Runtime/Private/AnimNext/AnimNextAlignment.h
+++ b/Engine/Plugins/Experimental/EvaluationNotifies/Source/Runtime/Private/AnimNext/AnimNextAlignment.h
@@ -45,5 +45,6 @@ USTRUCT()
 struct FEvaluationNotify_AlignToGroundInstance : public FEvaluationNotify_AlignmentInstance
 {
 	GENERATED_BODY()
+	virtual void End(UE::AnimNext::FEvaluationNotifiesTrait::FInstanceData& TraitInstanceData) override;
 	virtual bool GetTargetTransform(UE::AnimNext::FEvaluationNotifiesTrait::FInstanceData& TraitInstanceData, FTransform& TargetTranform);
 };
\ No newline at end of file
diff --git a/Engine/Plugins/Experimental/EvaluationNotifies/Source/Runtime/Private/AnimNext/EvaluationNotifiesTrait.cpp b/Engine/Plugins/Experimental/EvaluationNotifies/Source/Runtime/Private/AnimNext/EvaluationNotifiesTrait.cpp
--- a/Engine/Plugins/Experimental/EvaluationNotifies/Source/Runtime/Private/AnimNext/EvaluationNotifiesTrait.cpp
+++ b/Engine/Plugins/Experimental/EvaluationNotifies/Source/Runtime/Private/AnimNext/EvaluationNotifiesTrait.cpp
@@ -60,6 +60,7 @@ namespace UE::AnimNext
 		NotifySourceTrait.GetNotifies(Context, PreAdvanceState.GetPosition(), Delta.GetDeltaTime(), PreAdvanceState.IsLooping(), Notifies);
 
 		FInstanceData* InstanceData = Binding.GetInstanceData<FInstanceData>();
+		InstanceData->DeltaTime = Delta.GetDeltaTime();
 
 		for(const FAnimNotifyEventReference& NotifyData : Notifies)
 		{
@@ -75,7 +76,7 @@ namespace UE::AnimNext
 							return ExistingInstance.Get<FEvaluationNotify_BaseInstance>().NotifyEvent == NotifyEvent;
 						}))
 						{
-							FoundInstance->GetMutable<FEvaluationNotify_BaseInstance>().CurrentTime = NotifyData.GetCurrentAnimationTime() + DeltaTime;
+							FoundInstance->GetMutable<FEvaluationNotify_BaseInstance>().CurrentTime = NotifyData.GetCurrentAnimationTime();
 						}
 						else
 						{
@@ -94,17 +95,6 @@ namespace UE::AnimNext
 		}
 	}
 
-	void FEvaluationNotifiesTrait::PreUpdate(FUpdateTraversalContext& Context, const TTraitBinding<IUpdate>& Binding, const FTraitUpdateState& TraitState) const
-	{
-		FInstanceData* InstanceData = Binding.GetInstanceData<FInstanceData>();
-		check(InstanceData);
-
-		InstanceData->DeltaTime = TraitState.GetDeltaTime();
-
-		IUpdate::PreUpdate(Context, Binding, TraitState);
-
-	}
-
 	void FEvaluationNotifiesTrait::PostEvaluate(FEvaluateTraversalContext& Context, const TTraitBinding<IEvaluate>& Binding) const
 	{
 		IEvaluate::PostEvaluate(Context, Binding);
@@ -129,7 +119,7 @@ namespace UE::AnimNext
 			InstanceData->OnExtractRootMotionAttribute = AttributeTrait.GetOnExtractRootMotionAttribute(Context);
 		}
 		
-		InstanceData->DataInterface = Binding.GetTraitPtr().GetNodeInstance()->GetOwner().GetHost();
+		InstanceData->DataInterface = &Binding.GetTraitPtr().GetNodeInstance()->GetOwner();
 		InstanceData->DataInterface->GetVariable(RootBoneTransformName, InstanceData->RootBoneTransform);
 
 		Context.AppendTask(FAnimNextEvaluationNotifiesTask::Make(InstanceData, SharedData));
@@ -153,7 +143,7 @@ namespace UE::AnimNext
 			if (NotifyInstance.bActive)
 			{
 				NotifyInstance.bActive = false;
-				NotifyInstance.End();
+				NotifyInstance.End(*this);
 			}
 		}
 
@@ -200,10 +190,11 @@ void FAnimNextEvaluationNotifiesTask::Execute(UE::AnimNext::FEvaluationVM& VM) c
 				if (NotifyInstance.bActive)
 				{
 					NotifyInstance.bActive = false;
-					NotifyInstance.End();
+					NotifyInstance.End(*InstanceData);
 				}
 			}
-			
+
+			NotifyInstance.CurrentTime += InstanceData->DeltaTime;
 		}
 
 		InstanceData->EvaluationNotifies.RemoveAll([](const FInstancedStruct& Data)
diff --git a/Engine/Plugins/Experimental/EvaluationNotifies/Source/Runtime/Public/AnimNext/EvaluationNotifiesTrait.h b/Engine/Plugins/Experimental/EvaluationNotifies/Source/Runtime/Public/AnimNext/EvaluationNotifiesTrait.h
--- a/Engine/Plugins/Experimental/EvaluationNotifies/Source/Runtime/Public/AnimNext/EvaluationNotifiesTrait.h
+++ b/Engine/Plugins/Experimental/EvaluationNotifies/Source/Runtime/Public/AnimNext/EvaluationNotifiesTrait.h
@@ -59,7 +59,7 @@ struct EVALUATIONNOTIFIESRUNTIME_API FEvaluationNotifiesTrait : FAdditiveTrait,
 		/** the evaluation notifies extracted from the current anim sequence */
 		TArray<FInstancedStruct> EvaluationNotifies;
 
-		const FAnimNextDataInterfaceInstance* DataInterface = nullptr;
+		FAnimNextDataInterfaceInstance* DataInterface = nullptr;
 
 		/** Owner object the trait instance is associated to */
 		const UObject* HostObject;
@@ -68,9 +68,6 @@ struct EVALUATIONNOTIFIESRUNTIME_API FEvaluationNotifiesTrait : FAdditiveTrait,
 	// ITimelinePlayer impl
 	virtual void AdvanceBy(FExecutionContext& Context, const TTraitBinding<ITimelinePlayer>& Binding, float DeltaTime, bool bDispatchEvents) const override;
 	
-	// IUpdate impl 
-	virtual void PreUpdate(FUpdateTraversalContext& Context, const TTraitBinding<IUpdate>& Binding, const FTraitUpdateState& TraitState) const override;
-
 	// IEvaluate impl
 	virtual void PostEvaluate(FEvaluateTraversalContext& Context, const TTraitBinding<IEvaluate>& Binding) const override;
 
@@ -106,7 +103,7 @@ struct FEvaluationNotify_BaseInstance
 	GENERATED_BODY()
 	virtual void Start() {}
 	virtual void Update(UE::AnimNext::FEvaluationNotifiesTrait::FInstanceData& InstanceData, UE::AnimNext::FEvaluationVM& VM) {}
-	virtual void End() {}
+	virtual void End(UE::AnimNext::FEvaluationNotifiesTrait::FInstanceData& InstanceData) {}
 	
 	virtual ~FEvaluationNotify_BaseInstance() {};
 
diff --git a/Engine/Plugins/Experimental/EvaluationNotifies/Source/Runtime/Public/EvaluationNotifies/AnimNotifyState_Alignment.h b/Engine/Plugins/Experimental/EvaluationNotifies/Source/Runtime/Public/EvaluationNotifies/AnimNotifyState_Alignment.h
--- a/Engine/Plugins/Experimental/EvaluationNotifies/Source/Runtime/Public/EvaluationNotifies/AnimNotifyState_Alignment.h
+++ b/Engine/Plugins/Experimental/EvaluationNotifies/Source/Runtime/Public/EvaluationNotifies/AnimNotifyState_Alignment.h
@@ -170,14 +170,23 @@ class UNotifyState_AlignToGround : public UNotifyState_AlignmentBase
 	public:
 	
 	UPROPERTY(EditAnywhere, Category = "Trace Settings")
-	float TraceRadius;
+	float TraceRadius = 10;
 	
 	UPROPERTY(EditAnywhere, Category = "Trace Settings")
-	float TraceStartOffset;
+	float TraceStartOffset = -100;
 	
 	UPROPERTY(EditAnywhere, Category = "Trace Settings")
-	float TraceEndOffset;
+	float TraceEndOffset = 100;
 	
 	UPROPERTY(EditAnywhere, Category = "Trace Settings")
 	TEnumAsByte<ETraceTypeQuery> TraceChannel = TraceTypeQuery1;
+	
+	// Variable to output playback rate modifier to
+	UPROPERTY(EditAnywhere, Category = PlaybackRate)
+	FName PlaybackRateOutputVariable;
+	
+	UPROPERTY(EditAnywhere, Category = PlaybackRate)
+	float MinPlaybackRateModifier = 0.5;
+	UPROPERTY(EditAnywhere, Category = PlaybackRate)
+	float MaxPlaybackRateModifier = 1.0;
 };
\ No newline at end of file
