Expose some evaluation notify parameters

diff --git a/Engine/Plugins/Experimental/EvaluationNotifies/Source/Runtime/Public/EvaluationNotifies/AnimNotifyState_Alignment.h b/Engine/Plugins/Experimental/EvaluationNotifies/Source/Runtime/Public/EvaluationNotifies/AnimNotifyState_Alignment.h
--- a/Engine/Plugins/Experimental/EvaluationNotifies/Source/Runtime/Public/EvaluationNotifies/AnimNotifyState_Alignment.h
+++ b/Engine/Plugins/Experimental/EvaluationNotifies/Source/Runtime/Public/EvaluationNotifies/AnimNotifyState_Alignment.h
@@ -75,7 +75,7 @@ enum class EAlignmentUpdateMode
 	Relative,
 };
 
-UCLASS()
+UCLASS(MinimalAPI)
 class UNotifyState_AlignmentBase : public UAnimNotifyState
 {
 	GENERATED_BODY()
@@ -117,7 +117,7 @@ class UNotifyState_AlignmentBase : public UAnimNotifyState
 	FAlignmentSteeringSettings SteeringSettings;
 };
 
-UCLASS(BlueprintType, DisplayName="Alignment")
+UCLASS(MinimalAPI, BlueprintType, DisplayName="Alignment")
 class UNotifyState_Alignment : public UNotifyState_AlignmentBase
 {
 	GENERATED_BODY()
