add bool to disable alignment evaluationnotify

diff --git a/Engine/Plugins/Experimental/EvaluationNotifies/Source/Runtime/Private/AnimNext/AnimNextAlignment.cpp b/Engine/Plugins/Experimental/EvaluationNotifies/Source/Runtime/Private/AnimNext/AnimNextAlignment.cpp
--- a/Engine/Plugins/Experimental/EvaluationNotifies/Source/Runtime/Private/AnimNext/AnimNextAlignment.cpp
+++ b/Engine/Plugins/Experimental/EvaluationNotifies/Source/Runtime/Private/AnimNext/AnimNextAlignment.cpp
@@ -205,18 +205,35 @@ FTransform GetModelSpaceTransform(const UE::AnimNext::FLODPoseStack& Pose, FBone
 bool FEvaluationNotify_AlignmentInstance::GetTargetTransform(UE::AnimNext::FEvaluationNotifiesTrait::FInstanceData& TraitInstanceData, FTransform& Transform)
 {
 	const UNotifyState_Alignment* AlignmentNotify = CastChecked<UNotifyState_Alignment>(AnimNotify);
+	
 	return TraitInstanceData.DataInterface->GetVariable(AlignmentNotify->TransformName, Transform) == EPropertyBagResult::Success;
 }
 
 void FEvaluationNotify_AlignmentInstance::Update(UE::AnimNext::FEvaluationNotifiesTrait::FInstanceData& TraitInstanceData, UE::AnimNext::FEvaluationVM& VM)
 {
-	UNotifyState_AlignmentBase* AlignmentNotify = CastChecked<UNotifyState_AlignmentBase>(AnimNotify);
+	if (bDisabled)
+	{
+		return;
+	}
 	
+	UNotifyState_AlignmentBase* AlignmentNotify = CastChecked<UNotifyState_AlignmentBase>(AnimNotify);
+
 	if (const TUniquePtr<UE::AnimNext::FKeyframeState>* Keyframe = VM.PeekValue<TUniquePtr<UE::AnimNext::FKeyframeState>>(UE::AnimNext::KEYFRAME_STACK_NAME, 0))
 	{
 		if (bFirstFrame)
 		{
 			bFirstFrame = false;
+			
+			if (AlignmentNotify->Disable != NAME_None)
+			{
+				TraitInstanceData.DataInterface->GetVariable(AlignmentNotify->Disable, bDisabled);
+				
+				if (bDisabled)
+				{
+					return;
+				}
+			}
+			
 			ActualStartTime = CurrentTime;
 
 			const UE::AnimNext::FLODPoseStack& Pose = Keyframe->Get()->Pose;
diff --git a/Engine/Plugins/Experimental/EvaluationNotifies/Source/Runtime/Private/AnimNext/AnimNextAlignment.h b/Engine/Plugins/Experimental/EvaluationNotifies/Source/Runtime/Private/AnimNext/AnimNextAlignment.h
--- a/Engine/Plugins/Experimental/EvaluationNotifies/Source/Runtime/Private/AnimNext/AnimNextAlignment.h
+++ b/Engine/Plugins/Experimental/EvaluationNotifies/Source/Runtime/Private/AnimNext/AnimNextAlignment.h
@@ -27,6 +27,7 @@ struct FEvaluationNotify_AlignmentInstance : public FEvaluationNotify_BaseInstan
 	float GetWeight(float Time, const FAlignmentWarpCurve& WarpCurve) const; 
 
 	FBoneReference AlignBone;
+	bool bDisabled = false;
 	bool bFirstFrame = true;
 	float ActualStartTime = 0;
 	float PreviousFrame = 0;
diff --git a/Engine/Plugins/Experimental/EvaluationNotifies/Source/Runtime/Public/EvaluationNotifies/AnimNotifyState_Alignment.h b/Engine/Plugins/Experimental/EvaluationNotifies/Source/Runtime/Public/EvaluationNotifies/AnimNotifyState_Alignment.h
--- a/Engine/Plugins/Experimental/EvaluationNotifies/Source/Runtime/Public/EvaluationNotifies/AnimNotifyState_Alignment.h
+++ b/Engine/Plugins/Experimental/EvaluationNotifies/Source/Runtime/Public/EvaluationNotifies/AnimNotifyState_Alignment.h
@@ -108,6 +108,10 @@ class UNotifyState_AlignmentBase : public UAnimNotifyState
 		
 	UPROPERTY(EditAnywhere, Category = TargetTransform)
 	EAlignmentUpdateMode UpdateMode;
+	
+	// Bool variable which will disable this notify
+	UPROPERTY(EditAnywhere, Category = TargetTransform)
+	FName Disable;
 
 	// Steering will rotate the character based on the difference between the animated movement direction and the warped movement direction, to keep the character facing their direction of movement
 	UPROPERTY(EditAnywhere, Category = Steering)
