Separate translation curves for Alignment warping

diff --git a/Engine/Plugins/Experimental/EvaluationNotifies/Source/Runtime/Private/AnimNext/AnimNextAlignment.cpp b/Engine/Plugins/Experimental/EvaluationNotifies/Source/Runtime/Private/AnimNext/AnimNextAlignment.cpp
--- a/Engine/Plugins/Experimental/EvaluationNotifies/Source/Runtime/Private/AnimNext/AnimNextAlignment.cpp
+++ b/Engine/Plugins/Experimental/EvaluationNotifies/Source/Runtime/Private/AnimNext/AnimNextAlignment.cpp
@@ -345,21 +345,39 @@ void FEvaluationNotify_AlignmentInstance::Update(UE::AnimNext::FEvaluationNotifi
 				}
 
 				WarpedTrajectory.SetNum(AnimTrajectoryData.Trajectory.Num());
+				
+				FVector MovementDirection = TraitInstanceData.RootBoneTransform.GetRotation().GetRightVector(); // todo, get this from translation over first few frames or based on some setting?
 			
 				FTransform InverseLastFrame = AnimTrajectoryData.Trajectory.Last().Inverse();
 				// Translation warping + steering
 				for(int i=0;i<NumFrames; i++)
 				{
-					float Weight = GetWeight(ActualStartTime + AnimNextAlignment::FrameTime * (i+1), AlignmentNotify->TranslationWarpingCurve);
-
 					FTransform TransformFromRoot = AnimTrajectoryData.Trajectory[i] * TraitInstanceData.RootBoneTransform;
 					FTransform TransformFromTarget = AnimTrajectoryData.Trajectory[i] * InverseLastFrame * TargetTransform;
 				
 					FVector OldPosition = TransformFromRoot.GetTranslation();
 					FVector UnWarpedDelta = OldPosition - UnWarpedPreviousPosition;
 					UnWarpedPreviousPosition = OldPosition;
-				
-					FVector NewPosition = FMath::Lerp(OldPosition, TransformFromTarget.GetTranslation(), Weight);
+
+					FVector NewPosition = OldPosition;
+					
+					if (AlignmentNotify->bSeparateTranslationCurves)
+					{
+						float InMovementDirectionWeight = GetWeight(ActualStartTime + AnimNextAlignment::FrameTime*i, AlignmentNotify->TranslationWarpingCurve_InMovementDirection);
+						float OutOfMovementDirectionWeight = GetWeight(ActualStartTime + AnimNextAlignment::FrameTime*i, AlignmentNotify->TranslationWarpingCurve_OutOfMovementDirection);
+	
+						FVector Delta = TransformFromTarget.GetTranslation() - OldPosition;
+						FVector InMovementDirectionDelta = MovementDirection * Delta.Dot(MovementDirection);
+						FVector OutOfMovementDirectionDelta = Delta - InMovementDirectionDelta;
+	
+						NewPosition = OldPosition + InMovementDirectionDelta * InMovementDirectionWeight + OutOfMovementDirectionDelta * OutOfMovementDirectionWeight;
+					}
+					else
+					{
+						float Weight = GetWeight(ActualStartTime + AnimNextAlignment::FrameTime*i, AlignmentNotify->TranslationWarpingCurve);
+						NewPosition = FMath::Lerp(OldPosition, TransformFromTarget.GetTranslation(), Weight);
+					}
+
 					FVector WarpedDelta = NewPosition - WarpedPreviousPosition;
 					WarpedPreviousPosition = NewPosition;
 
diff --git a/Engine/Plugins/Experimental/EvaluationNotifies/Source/Runtime/Public/EvaluationNotifies/AnimNotifyState_Alignment.h b/Engine/Plugins/Experimental/EvaluationNotifies/Source/Runtime/Public/EvaluationNotifies/AnimNotifyState_Alignment.h
--- a/Engine/Plugins/Experimental/EvaluationNotifies/Source/Runtime/Public/EvaluationNotifies/AnimNotifyState_Alignment.h
+++ b/Engine/Plugins/Experimental/EvaluationNotifies/Source/Runtime/Public/EvaluationNotifies/AnimNotifyState_Alignment.h
@@ -84,6 +84,16 @@ class UNotifyState_AlignmentBase : public UAnimNotifyState
 
 	UPROPERTY(EditAnywhere, Category = Blending)
 	FAlignmentWarpCurve TranslationWarpingCurve;
+
+	
+	UPROPERTY(EditAnywhere, Category = Blending, Meta = (EditCondition="!bUseSeparateTranslationCurves"))
+	bool bSeparateTranslationCurves;
+
+	// curve for warping translation in direction of movement
+	UPROPERTY(EditAnywhere, Category = Blending, Meta = (EditCondition="bUseSeparateTranslationCurves"))
+	FAlignmentWarpCurve TranslationWarpingCurve_InMovementDirection;
+	UPROPERTY(EditAnywhere, Category = Blending, Meta = (EditCondition="bUseSeparateTranslationCurves"))
+	FAlignmentWarpCurve TranslationWarpingCurve_OutOfMovementDirection;
 	
 	UPROPERTY(EditAnywhere, Category = Blending)
 	FAlignmentWarpCurve RotationWarpingCurve = { };
