Add options for flattening alignment transforms

diff --git a/Engine/Plugins/Experimental/EvaluationNotifies/Source/Runtime/Private/AnimNext/AnimNextAlignment.cpp b/Engine/Plugins/Experimental/EvaluationNotifies/Source/Runtime/Private/AnimNext/AnimNextAlignment.cpp
--- a/Engine/Plugins/Experimental/EvaluationNotifies/Source/Runtime/Private/AnimNext/AnimNextAlignment.cpp
+++ b/Engine/Plugins/Experimental/EvaluationNotifies/Source/Runtime/Private/AnimNext/AnimNextAlignment.cpp
@@ -118,6 +118,13 @@ bool FEvaluationNotify_AlignToGroundInstance::GetTargetTransform(UE::AnimNext::F
 	const FCollisionShape CollisionShape = FCollisionShape::MakeSphere(AlignToGroundNotify->TraceRadius);
 
 	FTransform StartTransform = TraitInstanceData.RootBoneTransform;
+	if (AlignToGroundNotify->bForceStartTransformUpright)
+	{
+		FRotator Rotation = StartTransform.GetRotation().Rotator();
+		Rotation.Roll = 0;
+		Rotation.Pitch = 0;
+		StartTransform.SetRotation(Rotation.Quaternion());
+	}
 	
 	const float PredictionDelta = FMath::Max(EndTime - CurrentTime, 0);
 	FTransform PredictedRootMotion = TraitInstanceData.OnExtractRootMotionAttribute.Execute(CurrentTime, PredictionDelta, false);
@@ -260,6 +267,14 @@ void FEvaluationNotify_AlignmentInstance::Update(UE::AnimNext::FEvaluationNotifi
 				 }
 			
 				TargetTransform = AlignmentNotify->AlignOffset * TargetTransform;
+
+				if (AlignmentNotify->bForceTargetTransformUpright)
+				{
+					FRotator Rotation = TargetTransform.GetRotation().Rotator();
+					Rotation.Roll = 0;
+					Rotation.Pitch = 0;
+					TargetTransform.SetRotation(Rotation.Quaternion());
+				}
 			
 				StartingRootTransform = TraitInstanceData.RootBoneTransform;
 			
@@ -472,6 +487,7 @@ void FEvaluationNotify_AlignmentInstance::Update(UE::AnimNext::FEvaluationNotifi
 			for(int i=1; i<WarpedTrajectory.Num(); i++)
 			{
 				UE_VLOG_SEGMENT(TraitInstanceData.HostObject, "Alignment", Display, WarpedTrajectory[i-1].GetLocation(), WarpedTrajectory[i].GetLocation(), i%2 == 0 ? FColor::Green : FColor::Blue, TEXT(""));
+				UE_VLOG_SEGMENT(TraitInstanceData.HostObject, "Alignment", Display, WarpedTrajectory[i].GetLocation(), WarpedTrajectory[i].GetLocation() + WarpedTrajectory[i].GetRotation().GetUpVector() * 10, i%2 == 0 ? FColor::Yellow : FColor::Red, TEXT(""));
 			}
 
 			// a dot representing our current position on the trajectory 
diff --git a/Engine/Plugins/Experimental/EvaluationNotifies/Source/Runtime/Public/EvaluationNotifies/AnimNotifyState_Alignment.h b/Engine/Plugins/Experimental/EvaluationNotifies/Source/Runtime/Public/EvaluationNotifies/AnimNotifyState_Alignment.h
--- a/Engine/Plugins/Experimental/EvaluationNotifies/Source/Runtime/Public/EvaluationNotifies/AnimNotifyState_Alignment.h
+++ b/Engine/Plugins/Experimental/EvaluationNotifies/Source/Runtime/Public/EvaluationNotifies/AnimNotifyState_Alignment.h
@@ -98,11 +98,15 @@ class UNotifyState_AlignmentBase : public UAnimNotifyState
 	UPROPERTY(EditAnywhere, Category = Blending)
 	FAlignmentWarpCurve RotationWarpingCurve = { };
 	
-
 	// Offset from Root (or AlignBone) to Align to the target transform
 	UPROPERTY(EditAnywhere, Category = TargetTransform)
 	FTransform AlignOffset;
 
+	// only allow Yaw in the target transform (remove other rotations before warping)
+	UPROPERTY(EditAnywhere, Category = TargetTransform)
+	bool bForceTargetTransformUpright;
+	
+
 #if WITH_EDITOR
 	// For automating setting the AlignOffset at the end of the NotifyState
     UFUNCTION(CallInEditor, BlueprintCallable, Category = "TechAnimUtils")
@@ -203,4 +207,8 @@ class UNotifyState_AlignToGround : public UNotifyState_AlignmentBase
 	float MinPlaybackRateModifier = 0.5;
 	UPROPERTY(EditAnywhere, Category = PlaybackRate)
 	float MaxPlaybackRateModifier = 1.0;
+
+	// only allow Yaw in the starting root transform when computing the future root motion
+	UPROPERTY(EditAnywhere, Category = TraceSettings)
+	bool bForceStartTransformUpright;
 };
\ No newline at end of file
