Alignment bug fixes

diff --git a/Engine/Plugins/Experimental/EvaluationNotifies/Source/Runtime/Private/AnimNext/AnimNextAlignment.cpp b/Engine/Plugins/Experimental/EvaluationNotifies/Source/Runtime/Private/AnimNext/AnimNextAlignment.cpp
--- a/Engine/Plugins/Experimental/EvaluationNotifies/Source/Runtime/Private/AnimNext/AnimNextAlignment.cpp
+++ b/Engine/Plugins/Experimental/EvaluationNotifies/Source/Runtime/Private/AnimNext/AnimNextAlignment.cpp
@@ -61,21 +61,27 @@ float FEvaluationNotify_AlignmentInstance::GetWeight(float Time, const FAlignmen
 
 	if (WarpCurve.CurveType == EAlignmentWeightCurveType::FromRootMotionTranslation && !AnimTrajectoryData.TranslationCurve.IsEmpty())
 	{
-		const float Duration = (EndTime - StartTime);
+		const float Duration = (RoundedEndTime - StartTime);
 		const float StartCurveSampleTime = FMath::Max(WarpCurve.StartRatio * Duration, ActualStartTime) - ActualStartTime;
 		return AnimNextAlignment::SampleCurve(Time - ActualStartTime, StartCurveSampleTime, WarpCurve.EndRatio * Duration - ActualStartTime, AnimTrajectoryData.TranslationCurve);
 	}
 	else if (WarpCurve.CurveType == EAlignmentWeightCurveType::FromRootMotionRotation && !AnimTrajectoryData.RotationCurve.IsEmpty())
 	{
-		const float Duration = (EndTime - StartTime);
+		const float Duration = (RoundedEndTime - StartTime);
 		const float StartCurveSampleTime = FMath::Max(WarpCurve.StartRatio * Duration, ActualStartTime) - ActualStartTime;
 		return AnimNextAlignment::SampleCurve(Time - ActualStartTime, StartCurveSampleTime, WarpCurve.EndRatio * Duration - ActualStartTime, AnimTrajectoryData.RotationCurve);
 	}
 	else
 	{
-		const float StartTimeWithRatio = FMath::Max(FMath::Lerp(StartTime, EndTime, WarpCurve.StartRatio), ActualStartTime);
-		const float EndTimeWithRatio = FMath::Lerp(StartTime, EndTime, WarpCurve.EndRatio);
-			
+		const float StartTimeWithRatio = FMath::Max(FMath::Lerp(StartTime, RoundedEndTime, WarpCurve.StartRatio), ActualStartTime);
+		const float EndTimeWithRatio = FMath::Lerp(StartTime, RoundedEndTime, WarpCurve.EndRatio);
+
+		if(EndTimeWithRatio<=StartTimeWithRatio)
+		{
+			// the animation likely started late, and the blend ratio is finished so we should snap
+			return 1.0f;
+		}
+		
 		const float CurrentRelativeTime = FMath::Clamp((Time-StartTimeWithRatio)/(EndTimeWithRatio-StartTimeWithRatio), 0,1);
 		
 		switch(WarpCurve.CurveType)
@@ -126,7 +132,7 @@ bool FEvaluationNotify_AlignToGroundInstance::GetTargetTransform(UE::AnimNext::F
 		StartTransform.SetRotation(Rotation.Quaternion());
 	}
 	
-	const float PredictionDelta = FMath::Max(EndTime - CurrentTime, 0);
+	const float PredictionDelta = FMath::Max(RoundedEndTime - CurrentTime, 0);
 	FTransform PredictedRootMotion = TraitInstanceData.OnExtractRootMotionAttribute.Execute(CurrentTime, PredictionDelta, false);
 	FTransform EndTransform = PredictedRootMotion * StartTransform;
 
@@ -160,7 +166,7 @@ bool FEvaluationNotify_AlignToGroundInstance::GetTargetTransform(UE::AnimNext::F
 		if (bHit)
 		{
 			Transform = EndTransform;
-			Transform.SetLocation(HitResult.ImpactPoint);
+			Transform.SetLocation(HitResult.ImpactPoint + FVector(0,0,AlignToGroundNotify->GroundHeightOffset));
 		
 			static const FBox UnitBox(FVector(-10, -10, -10), FVector(10, 10, 10));
 			UE_VLOG_OBOX(TraitInstanceData.HostObject, "AlignToGround", Display, UnitBox, Transform.ToMatrixWithScale(), FColor::Red, TEXT(""));
@@ -243,7 +249,8 @@ void FEvaluationNotify_AlignmentInstance::Update(UE::AnimNext::FEvaluationNotifi
 			}
 			
 			ActualStartTime = CurrentTime;
-
+			RoundedEndTime = ActualStartTime + FMath::CeilToDouble((EndTime - ActualStartTime) / AnimNextAlignment::FrameTime) * AnimNextAlignment::FrameTime;
+			
 			const UE::AnimNext::FLODPoseStack& Pose = Keyframe->Get()->Pose;
 			AlignBone.Initialize(Pose.GetSkeletonAsset());
 
@@ -253,7 +260,7 @@ void FEvaluationNotify_AlignmentInstance::Update(UE::AnimNext::FEvaluationNotifi
 				UE_VLOG_OBOX(TraitInstanceData.HostObject, "Alignment", Display, UnitBox, TargetTransform.ToMatrixWithScale(), FColor::Blue, TEXT(""));
 			
 				// get alignment transform
-				const float PredictionDelta = EndTime - CurrentTime;
+				const float PredictionDelta = RoundedEndTime - CurrentTime;
 
 				if (AlignBone.HasValidSetup())
 				{
@@ -280,7 +287,7 @@ void FEvaluationNotify_AlignmentInstance::Update(UE::AnimNext::FEvaluationNotifi
 			
 				// extract root motion trajectory todo: this should be cached and reused
 			
-				int NumFrames = 1 + (EndTime - CurrentTime) / AnimNextAlignment::FrameTime;
+				int NumFrames = 1 + (RoundedEndTime - ActualStartTime) / AnimNextAlignment::FrameTime;
 
 				if (NumFrames <= 0)
 				{
@@ -295,8 +302,12 @@ void FEvaluationNotify_AlignmentInstance::Update(UE::AnimNext::FEvaluationNotifi
 			
 				FTransform PredictedTransform;
 				float PredictionTime = CurrentTime;
+				
+				AnimTrajectoryData.Trajectory[0] = PredictedTransform;
+				PredictionTime += AnimNextAlignment::FrameTime;
+				AnimTrajectoryData.TranslationCurve[0] = 0;
 
-				for(int i=0;i<NumFrames; i++)
+				for(int i=1;i<NumFrames; i++)
 				{
 					FTransform RootMotionThisFrame = TraitInstanceData.OnExtractRootMotionAttribute.Execute(PredictionTime, AnimNextAlignment::FrameTime, false);
 				
@@ -304,17 +315,8 @@ void FEvaluationNotify_AlignmentInstance::Update(UE::AnimNext::FEvaluationNotifi
 					AnimTrajectoryData.Trajectory[i] = PredictedTransform;
 					PredictionTime += AnimNextAlignment::FrameTime;
 
-					AnimTrajectoryData.TranslationCurve[i] = RootMotionThisFrame.GetTranslation().Length();
-					if (i>0)
-					{
-						AnimTrajectoryData.TranslationCurve[i] += AnimTrajectoryData.TranslationCurve[i-1]; 
-					}
-				
-					AnimTrajectoryData.RotationCurve[i] = fabs(RootMotionThisFrame.GetRotation().GetAngle());
-					if (i>0)
-					{
-						AnimTrajectoryData.RotationCurve[i] += AnimTrajectoryData.RotationCurve[i-1]; 
-					}
+					AnimTrajectoryData.TranslationCurve[i] = RootMotionThisFrame.GetTranslation().Length() + AnimTrajectoryData.TranslationCurve[i-1]; 
+					AnimTrajectoryData.RotationCurve[i] = fabs(RootMotionThisFrame.GetRotation().GetAngle()) + AnimTrajectoryData.RotationCurve[i-1]; 
 				}
 
 				// normalize curves;
@@ -403,7 +405,7 @@ void FEvaluationNotify_AlignmentInstance::Update(UE::AnimNext::FEvaluationNotifi
 					if (i > 0 && AlignmentNotify->bEnableSteering)
 					{
 						FQuat OldRotation = TransformFromRoot.GetRotation();
-						FQuat DirectionChange = FQuat::FindBetweenVectors(UnWarpedDelta, WarpedDelta);
+						FQuat DirectionChange = FQuat::FindBetweenVectors(FVector(UnWarpedDelta.X, UnWarpedDelta.Y, 0) , FVector(WarpedDelta.X, WarpedDelta.Y, 0));
 
 						if (AlignmentNotify->SteeringSettings.bEnableSmoothing)
 						{
@@ -444,54 +446,71 @@ void FEvaluationNotify_AlignmentInstance::Update(UE::AnimNext::FEvaluationNotifi
 	
 		const UE::Anim::IAnimRootMotionProvider* RootMotionProvider = UE::Anim::IAnimRootMotionProvider::Get();
 		ensureMsgf(RootMotionProvider, TEXT("Alignment expected a valid root motion delta provider interface."));
-	
     		
 		if (RootMotionProvider && !WarpedTrajectory.IsEmpty())
 		{
-			float Frame = (CurrentTime - ActualStartTime) / AnimNextAlignment::FrameTime;
-			FTransform WorldTransform;
-			AnimNextAlignment::GetTransformForFrame(Frame, WarpedTrajectory, WorldTransform);
-
-			if(AlignmentNotify->UpdateMode == EAlignmentUpdateMode::World)
-			{
-				RootMotionProvider->OverrideRootMotion(WorldTransform.GetRelativeTransform(TraitInstanceData.RootBoneTransform), Keyframe->Get()->Attributes);
-			}
-			else // relative mode
+			float Frame = (CurrentTime + TraitInstanceData.DeltaTime - ActualStartTime) / AnimNextAlignment::FrameTime;
+			
+			if (PreviousFrame < WarpedTrajectory.Num() - 1)
 			{
-				FTransform PrevTransform;
-				AnimNextAlignment::GetTransformForFrame(PreviousFrame, WarpedTrajectory, PrevTransform);
-				RootMotionProvider->OverrideRootMotion(WorldTransform.GetRelativeTransform(PrevTransform), Keyframe->Get()->Attributes);
-				PreviousFrame = Frame;
-			}
+				FTransform WorldTransform;
+				
+				if (Frame > WarpedTrajectory.Num()-1)
+				{
+					float PostAlignmentTime = (Frame - (WarpedTrajectory.Num() - 1)) * AnimNextAlignment::FrameTime;
+					FTransform PostAlignmentRootMotion = TraitInstanceData.OnExtractRootMotionAttribute.Execute(RoundedEndTime, PostAlignmentTime, false);
+					AnimNextAlignment::GetTransformForFrame(Frame, WarpedTrajectory, WorldTransform);
+					WorldTransform = PostAlignmentRootMotion * TargetTransform;
+				}
+				else
+				{
+					AnimNextAlignment::GetTransformForFrame(Frame, WarpedTrajectory, WorldTransform);
+				}
+
+				if(AlignmentNotify->UpdateMode == EAlignmentUpdateMode::World)
+				{
+					RootMotionProvider->OverrideRootMotion(WorldTransform.GetRelativeTransform(TraitInstanceData.RootBoneTransform), Keyframe->Get()->Attributes);
+				}
+				else // relative mode
+				{
+					FTransform PrevTransform;
+					AnimNextAlignment::GetTransformForFrame(PreviousFrame, WarpedTrajectory, PrevTransform);
+					RootMotionProvider->OverrideRootMotion(WorldTransform.GetRelativeTransform(PrevTransform), Keyframe->Get()->Attributes);
+				}
 		
-			// unwarped trajectory relative to starting transform
-			FTransform PreviousTransform = AnimTrajectoryData.Trajectory[0] * StartingRootTransform;
-			for(int i=1;i<AnimTrajectoryData.Trajectory.Num(); i++)
-			{
-				FTransform TransformFromRoot = AnimTrajectoryData.Trajectory[i] * StartingRootTransform;
-				UE_VLOG_SEGMENT(TraitInstanceData.HostObject, "Alignment", Display, PreviousTransform.GetLocation(), TransformFromRoot.GetLocation(), i%2 == 0 ? FColor::Yellow : FColor::Red, TEXT(""));
-				PreviousTransform = TransformFromRoot;
-			}
+				// unwarped trajectory relative to starting transform
+				FTransform PreviousTransform = AnimTrajectoryData.Trajectory[0] * StartingRootTransform;
+				for(int i=1;i<AnimTrajectoryData.Trajectory.Num(); i++)
+				{
+					FTransform TransformFromRoot = AnimTrajectoryData.Trajectory[i] * StartingRootTransform;
+					UE_VLOG_SEGMENT(TraitInstanceData.HostObject, "Alignment", Display, PreviousTransform.GetLocation(), TransformFromRoot.GetLocation(), i%2 == 0 ? FColor::Yellow : FColor::Red, TEXT(""));
+					PreviousTransform = TransformFromRoot;
+				}
 		
-			// unwarped trajectory relative to target transform
-			FTransform InverseLastFrame = AnimTrajectoryData.Trajectory.Last().Inverse();
-			PreviousTransform = AnimTrajectoryData.Trajectory[0] * InverseLastFrame * TargetTransform;
-			for(int i=1;i<AnimTrajectoryData.Trajectory.Num(); i++)
-			{
-				FTransform TransformFromTarget = AnimTrajectoryData.Trajectory[i] * InverseLastFrame * TargetTransform;
-				UE_VLOG_SEGMENT(TraitInstanceData.HostObject, "Alignment", Display, PreviousTransform.GetLocation(), TransformFromTarget.GetLocation(), i%2 == 0 ? FColor::Yellow : FColor::Red, TEXT(""));
-				PreviousTransform = TransformFromTarget;
-			}
+				// unwarped trajectory relative to target transform
+				FTransform InverseLastFrame = AnimTrajectoryData.Trajectory.Last().Inverse();
+				PreviousTransform = AnimTrajectoryData.Trajectory[0] * InverseLastFrame * TargetTransform;
+				for(int i=1;i<AnimTrajectoryData.Trajectory.Num(); i++)
+				{
+					FTransform TransformFromTarget = AnimTrajectoryData.Trajectory[i] * InverseLastFrame * TargetTransform;
+					UE_VLOG_SEGMENT(TraitInstanceData.HostObject, "Alignment", Display, PreviousTransform.GetLocation(), TransformFromTarget.GetLocation(), i%2 == 0 ? FColor::Yellow : FColor::Red, TEXT(""));
+					PreviousTransform = TransformFromTarget;
+				}
 		
-			// the warped trajectory
-			for(int i=1; i<WarpedTrajectory.Num(); i++)
-			{
-				UE_VLOG_SEGMENT(TraitInstanceData.HostObject, "Alignment", Display, WarpedTrajectory[i-1].GetLocation(), WarpedTrajectory[i].GetLocation(), i%2 == 0 ? FColor::Green : FColor::Blue, TEXT(""));
-				UE_VLOG_SEGMENT(TraitInstanceData.HostObject, "Alignment", Display, WarpedTrajectory[i].GetLocation(), WarpedTrajectory[i].GetLocation() + WarpedTrajectory[i].GetRotation().GetUpVector() * 10, i%2 == 0 ? FColor::Yellow : FColor::Red, TEXT(""));
-			}
+				// the warped trajectory
+				for(int i=1; i<WarpedTrajectory.Num(); i++)
+				{
+					UE_VLOG_SEGMENT(TraitInstanceData.HostObject, "Alignment", Display, WarpedTrajectory[i-1].GetLocation(), WarpedTrajectory[i].GetLocation(), i%2 == 0 ? FColor::Green : FColor::Blue, TEXT(""));
+				}
 
-			// a dot representing our current position on the trajectory 
-			UE_VLOG_SPHERE(TraitInstanceData.HostObject, "Alignment", Display, WorldTransform.GetLocation(),  1, FColor::Red, TEXT(""));
+				// a dot representing our current root position
+				UE_VLOG_SPHERE(TraitInstanceData.HostObject, "Alignment", Display, TraitInstanceData.RootBoneTransform.GetLocation(),  1, FColor::Green, TEXT(""));
+				
+				// a dot representing our current target position on the trajectory 
+				UE_VLOG_SPHERE(TraitInstanceData.HostObject, "Alignment", Display, WorldTransform.GetLocation(),  1, FColor::Red, TEXT(""));
+				
+				PreviousFrame = Frame;
+			}
 		}
 	}
 }
diff --git a/Engine/Plugins/Experimental/EvaluationNotifies/Source/Runtime/Private/AnimNext/AnimNextAlignment.h b/Engine/Plugins/Experimental/EvaluationNotifies/Source/Runtime/Private/AnimNext/AnimNextAlignment.h
--- a/Engine/Plugins/Experimental/EvaluationNotifies/Source/Runtime/Private/AnimNext/AnimNextAlignment.h
+++ b/Engine/Plugins/Experimental/EvaluationNotifies/Source/Runtime/Private/AnimNext/AnimNextAlignment.h
@@ -30,6 +30,7 @@ struct FEvaluationNotify_AlignmentInstance : public FEvaluationNotify_BaseInstan
 	bool bDisabled = false;
 	bool bFirstFrame = true;
 	float ActualStartTime = 0;
+	float RoundedEndTime = 0;
 	float PreviousFrame = 0;
 
 	FTransform StartingRootTransform;
diff --git a/Engine/Plugins/Experimental/EvaluationNotifies/Source/Runtime/Private/AnimNext/EvaluationNotifiesTrait.cpp b/Engine/Plugins/Experimental/EvaluationNotifies/Source/Runtime/Private/AnimNext/EvaluationNotifiesTrait.cpp
--- a/Engine/Plugins/Experimental/EvaluationNotifies/Source/Runtime/Private/AnimNext/EvaluationNotifiesTrait.cpp
+++ b/Engine/Plugins/Experimental/EvaluationNotifies/Source/Runtime/Private/AnimNext/EvaluationNotifiesTrait.cpp
@@ -189,6 +189,7 @@ void FAnimNextEvaluationNotifiesTask::Execute(UE::AnimNext::FEvaluationVM& VM) c
 			{
 				if (NotifyInstance.bActive)
 				{
+					NotifyInstance.Update(*InstanceData, VM);
 					NotifyInstance.bActive = false;
 					NotifyInstance.End(*InstanceData);
 				}
diff --git a/Engine/Plugins/Experimental/EvaluationNotifies/Source/Runtime/Public/EvaluationNotifies/AnimNotifyState_Alignment.h b/Engine/Plugins/Experimental/EvaluationNotifies/Source/Runtime/Public/EvaluationNotifies/AnimNotifyState_Alignment.h
--- a/Engine/Plugins/Experimental/EvaluationNotifies/Source/Runtime/Public/EvaluationNotifies/AnimNotifyState_Alignment.h
+++ b/Engine/Plugins/Experimental/EvaluationNotifies/Source/Runtime/Public/EvaluationNotifies/AnimNotifyState_Alignment.h
@@ -195,6 +195,10 @@ class UNotifyState_AlignToGround : public UNotifyState_AlignmentBase
 	
 	UPROPERTY(EditAnywhere, Category = "Trace Settings")
 	float TraceEndOffset = 100;
+
+	// vertical offset added to the raycast hit result height
+	UPROPERTY(EditAnywhere, Category = "Trace Settings")
+	float GroundHeightOffset = 0;
 	
 	UPROPERTY(EditAnywhere, Category = "Trace Settings")
 	TEnumAsByte<ETraceTypeQuery> TraceChannel = TraceTypeQuery1;
