UAF: Namespace renamed

diff --git a/Engine/Plugins/Experimental/AnimNextControlRig/Source/AnimNextControlRig/Internal/ControlRigTrait.h b/Engine/Plugins/Experimental/AnimNextControlRig/Source/AnimNextControlRig/Internal/ControlRigTrait.h
--- a/Engine/Plugins/Experimental/AnimNextControlRig/Source/AnimNextControlRig/Internal/ControlRigTrait.h
+++ b/Engine/Plugins/Experimental/AnimNextControlRig/Source/AnimNextControlRig/Internal/ControlRigTrait.h
@@ -159,16 +159,16 @@ struct FControlRigTraitSharedData : public FAnimNextTraitSharedData
 	TArray<uint32> LatentPropertyMemoryLayouts;
 
 	// Manual handling of latent properties
-	static void ConstructLatentProperties(const UE::AnimNext::FTraitBinding& Binding);
+	static void ConstructLatentProperties(const UE::UAF::FTraitBinding& Binding);
 
-	static void DestructLatentProperties(const UE::AnimNext::FTraitBinding& Binding);
+	static void DestructLatentProperties(const UE::UAF::FTraitBinding& Binding);
 
 #if WITH_EDITOR
 	ANIMNEXTCONTROLRIG_API USkeleton* GetPreviewSkeleton() const;
 #endif // WITH_EDITOR
 };
 
-namespace UE::AnimNext
+namespace UE::UAF
 {
 
 struct FControlRigTrait : FBaseTrait, IEvaluate, IUpdate, IHierarchy, IGarbageCollection
@@ -224,8 +224,8 @@ struct FControlRigTrait : FBaseTrait, IEvaluate, IUpdate, IHierarchy, IGarbageCo
 		static UClass* GetTargetClass(const FSharedData* SharedData);
 
 		// Returns the property, variable name and the memory of the variable (as target) and the memory of the latent property (as source)
-		static int32 GetExposedVariablesData(const UE::AnimNext::FTraitBinding& Binding, const FSharedData* SharedData, FControlRigVariableMappings::FCustomPropertyMappings& OutMappings);
-		static int32 GetExposedControlsData(const UE::AnimNext::FTraitBinding& Binding, const FSharedData* SharedData, FControlRigVariableMappings::FCustomPropertyMappings& OutMappings);
+		static int32 GetExposedVariablesData(const UE::UAF::FTraitBinding& Binding, const FSharedData* SharedData, FControlRigVariableMappings::FCustomPropertyMappings& OutMappings);
+		static int32 GetExposedControlsData(const UE::UAF::FTraitBinding& Binding, const FSharedData* SharedData, FControlRigVariableMappings::FCustomPropertyMappings& OutMappings);
 
 		void HandleOnInitialized_AnyThread(URigVMHost*, const FName&);
 
@@ -242,7 +242,7 @@ struct FControlRigTrait : FBaseTrait, IEvaluate, IUpdate, IHierarchy, IGarbageCo
 	// FTrait impl
 	virtual void SerializeTraitSharedData(FArchive& Ar, FAnimNextTraitSharedData& SharedData) const override;
 	virtual uint32 GetNumLatentTraitProperties() const override;
-	virtual UE::AnimNext::FTraitLatentPropertyMemoryLayout GetLatentPropertyMemoryLayout(const FAnimNextTraitSharedData& SharedData, FName PropertyName, uint32 PropertyIndex) const override;
+	virtual UE::UAF::FTraitLatentPropertyMemoryLayout GetLatentPropertyMemoryLayout(const FAnimNextTraitSharedData& SharedData, FName PropertyName, uint32 PropertyIndex) const override;
 
 	// IUpdate impl
 	virtual void OnBecomeRelevant(FUpdateTraversalContext& Context, const TTraitBinding<IUpdate>& Binding, const FTraitUpdateState& TraitState) const override;
@@ -289,4 +289,4 @@ protected:
 	static bool GetControlSizeAndAlignment(ERigControlType ControlType, uint32& PropertySize, uint32& PropertyAlignment);
 };
 
-} // namespace UE::AnimNext
+} // namespace UE::UAF
diff --git a/Engine/Plugins/Experimental/AnimNextControlRig/Source/AnimNextControlRig/Private/AnimNextControlRigHierarchyMappings.cpp b/Engine/Plugins/Experimental/AnimNextControlRig/Source/AnimNextControlRig/Private/AnimNextControlRigHierarchyMappings.cpp
--- a/Engine/Plugins/Experimental/AnimNextControlRig/Source/AnimNextControlRig/Private/AnimNextControlRigHierarchyMappings.cpp
+++ b/Engine/Plugins/Experimental/AnimNextControlRig/Source/AnimNextControlRig/Private/AnimNextControlRigHierarchyMappings.cpp
@@ -17,7 +17,7 @@
 #include "Engine/SkeletalMesh.h"
 #include "Animation/Skeleton.h"
 
-namespace UE::AnimNext::ControlRig
+namespace UE::UAF::ControlRig
 {
 
 void FAnimNextControlRigHierarchyMappings::InitializeInstance()
@@ -36,7 +36,7 @@ void FAnimNextControlRigHierarchyMappings::LinkToHierarchy(URigHierarchy* InHier
 void FAnimNextControlRigHierarchyMappings::UpdateControlRigRefPoseIfNeeded(UControlRig* ControlRig
 	, const UObject* InstanceObject
 	, const USkeletalMeshComponent* SkeletalMeshComponent
-	, const UE::AnimNext::FReferencePose& InRefPose
+	, const UE::UAF::FReferencePose& InRefPose
 	, bool bInSetRefPoseFromSkeleton
 	, bool bIncludePoseInHash)
 {
@@ -84,7 +84,7 @@ void FAnimNextControlRigHierarchyMappings::UpdateControlRigRefPoseIfNeeded(UCont
 
 void FAnimNextControlRigHierarchyMappings::UpdateInputOutputMappingIfRequired(UControlRig* InControlRig
 	, URigHierarchy* InHierarchy
-	, const UE::AnimNext::FReferencePose& InRefPose
+	, const UE::UAF::FReferencePose& InRefPose
 	, int32 InCurrentLOD
 	, const TArray<FBoneReference>& InInputBonesToTransfer
 	, const TArray<FBoneReference>& InOutputBonesToTransfer
@@ -116,7 +116,7 @@ void FAnimNextControlRigHierarchyMappings::UpdateInputOutputMappingIfRequired(UC
 }
 
 void FAnimNextControlRigHierarchyMappings::UpdateInput(UControlRig* ControlRig
-	, UE::AnimNext::FKeyframeState& InOutput
+	, UE::UAF::FKeyframeState& InOutput
 	, const FControlRigIOSettings& InInputSettings
 	, const FControlRigIOSettings& InOutputSettings
 	, TWeakObjectPtr<UNodeMappingContainer>& InNodeMappingContainer
@@ -197,7 +197,7 @@ void FAnimNextControlRigHierarchyMappings::UpdateInput(UControlRig* ControlRig
 }
 
 void FAnimNextControlRigHierarchyMappings::UpdateOutput(UControlRig* ControlRig
-	, UE::AnimNext::FKeyframeState& InOutput
+	, UE::UAF::FKeyframeState& InOutput
 	, const FControlRigIOSettings& InOutputSettings
 	, TWeakObjectPtr<UNodeMappingContainer>& InNodeMappingContainer
 	, bool bInExecute
@@ -320,7 +320,7 @@ bool FAnimNextControlRigHierarchyMappings::IsUpdateToDate(const URigHierarchy* I
 
 void FAnimNextControlRigHierarchyMappings::PerformUpdateToDate(UControlRig* ControlRig
 	, URigHierarchy* InHierarchy
-	, const UE::AnimNext::FReferencePose& InRefPose
+	, const UE::UAF::FReferencePose& InRefPose
 	, int32 InCurrentLOD
 	, const TWeakObjectPtr<UNodeMappingContainer>& InNodeMappingContainer
 	, bool bInTransferPoseInGlobalSpace
@@ -336,4 +336,4 @@ void FAnimNextControlRigHierarchyMappings::PerformUpdateToDate(UControlRig* Cont
 	}
 }
 
-} // namespace UE::AnimNext::ControlRig
+} // namespace UE::UAF::ControlRig
diff --git a/Engine/Plugins/Experimental/AnimNextControlRig/Source/AnimNextControlRig/Private/AnimNextControlRigModule.cpp b/Engine/Plugins/Experimental/AnimNextControlRig/Source/AnimNextControlRig/Private/AnimNextControlRigModule.cpp
--- a/Engine/Plugins/Experimental/AnimNextControlRig/Source/AnimNextControlRig/Private/AnimNextControlRigModule.cpp
+++ b/Engine/Plugins/Experimental/AnimNextControlRig/Source/AnimNextControlRig/Private/AnimNextControlRigModule.cpp
@@ -4,7 +4,7 @@
 #include "Modules/ModuleInterface.h"
 #include "Modules/ModuleManager.h"
 
-namespace UE::AnimNext::ControlRig
+namespace UE::UAF::ControlRig
 {
 
 #if WITH_EDITOR
@@ -34,4 +34,4 @@ void FAnimNextControlRigModule::ShutdownModule()
 
 } // end namespace
 
-IMPLEMENT_MODULE(UE::AnimNext::ControlRig::FAnimNextControlRigModule, AnimNextControlRig)
+IMPLEMENT_MODULE(UE::UAF::ControlRig::FAnimNextControlRigModule, AnimNextControlRig)
diff --git a/Engine/Plugins/Experimental/AnimNextControlRig/Source/AnimNextControlRig/Private/AnimNextControlRigPoseAdapter.cpp b/Engine/Plugins/Experimental/AnimNextControlRig/Source/AnimNextControlRig/Private/AnimNextControlRigPoseAdapter.cpp
--- a/Engine/Plugins/Experimental/AnimNextControlRig/Source/AnimNextControlRig/Private/AnimNextControlRigPoseAdapter.cpp
+++ b/Engine/Plugins/Experimental/AnimNextControlRig/Source/AnimNextControlRig/Private/AnimNextControlRigPoseAdapter.cpp
@@ -6,12 +6,12 @@
 #include "Animation/NodeMappingContainer.h"
 #include "Engine/SkeletalMesh.h"
 
-namespace UE::AnimNext::ControlRig
+namespace UE::UAF::ControlRig
 {
 
 void FAnimNextControlRigPoseAdapter::UpdateInputOutputMappingIfRequired(UControlRig* InControlRig
 	, URigHierarchy* InHierarchy
-	, const UE::AnimNext::FReferencePose& InRefPose
+	, const UE::UAF::FReferencePose& InRefPose
 	, const int32 InCurentLOD
 	, const TWeakObjectPtr<UNodeMappingContainer>& InNodeMappingContainer
 	, bool bInTransferPoseInGlobalSpace
@@ -195,7 +195,7 @@ void FAnimNextControlRigPoseAdapter::UpdateInputOutputMappingIfRequired(UControl
 	}
 }
 
-const FReferenceSkeleton* FAnimNextControlRigPoseAdapter::GetReferenceSkeleton(const UE::AnimNext::FReferencePose& InRefPose)
+const FReferenceSkeleton* FAnimNextControlRigPoseAdapter::GetReferenceSkeleton(const UE::UAF::FReferencePose& InRefPose)
 {
 	const FReferenceSkeleton* ReferenceSkeleton = nullptr;
 
@@ -211,4 +211,4 @@ const FReferenceSkeleton* FAnimNextControlRigPoseAdapter::GetReferenceSkeleton(c
 	return ReferenceSkeleton;
 }
 
-} // namespace UE::AnimNext::ControlRig
+} // namespace UE::UAF::ControlRig
diff --git a/Engine/Plugins/Experimental/AnimNextControlRig/Source/AnimNextControlRig/Private/ControlRigTask.cpp b/Engine/Plugins/Experimental/AnimNextControlRig/Source/AnimNextControlRig/Private/ControlRigTask.cpp
--- a/Engine/Plugins/Experimental/AnimNextControlRig/Source/AnimNextControlRig/Private/ControlRigTask.cpp
+++ b/Engine/Plugins/Experimental/AnimNextControlRig/Source/AnimNextControlRig/Private/ControlRigTask.cpp
@@ -21,7 +21,7 @@ static TAutoConsoleVariable<int32> CVarControlRigDisableExecutionAnimNext(TEXT("
 
 DEFINE_STAT(STAT_AnimNext_Task_ControlRig);
 
-FAnimNextControlRigTask FAnimNextControlRigTask::Make(const UE::AnimNext::FControlRigTrait::FSharedData* SharedData, UE::AnimNext::FControlRigTrait::FInstanceData* InstanceData)
+FAnimNextControlRigTask FAnimNextControlRigTask::Make(const UE::UAF::FControlRigTrait::FSharedData* SharedData, UE::UAF::FControlRigTrait::FInstanceData* InstanceData)
 {
 	FAnimNextControlRigTask Task;
 	Task.InstanceData = InstanceData;
@@ -29,17 +29,17 @@ FAnimNextControlRigTask FAnimNextControlRigTask::Make(const UE::AnimNext::FContr
 	return Task;
 }
 
-void FAnimNextControlRigTask::Execute(UE::AnimNext::FEvaluationVM& VM) const
+void FAnimNextControlRigTask::Execute(UE::UAF::FEvaluationVM& VM) const
 {
 	SCOPE_CYCLE_COUNTER(STAT_AnimNext_Task_ControlRig);
 	DECLARE_SCOPE_HIERARCHICAL_COUNTER_FUNC()
 
-	using namespace UE::AnimNext;
+	using namespace UE::UAF;
 
 	TUniquePtr<FKeyframeState> KeyframeOut;
 
 	// Try to get a Keyframe from the stack
-	if (!VM.PopValue(UE::AnimNext::KEYFRAME_STACK_NAME, KeyframeOut))
+	if (!VM.PopValue(UE::UAF::KEYFRAME_STACK_NAME, KeyframeOut))
 	{
 		// if there is none, just create a reference frame
 		KeyframeOut = MakeUnique<FKeyframeState>(VM.MakeReferenceKeyframe(false));
@@ -53,7 +53,7 @@ void FAnimNextControlRigTask::Execute(UE::AnimNext::FEvaluationVM& VM) const
 	InstanceData->LastLOD = VM.GetCurrentLOD();
 }
 
-void FAnimNextControlRigTask::ExecuteControlRig(UE::AnimNext::FEvaluationVM& VM, UE::AnimNext::FKeyframeState& KeyFrameState, UControlRig* ControlRig) const
+void FAnimNextControlRigTask::ExecuteControlRig(UE::UAF::FEvaluationVM& VM, UE::UAF::FKeyframeState& KeyFrameState, UControlRig* ControlRig) const
 {
 	DECLARE_SCOPE_HIERARCHICAL_COUNTER_FUNC()
 
@@ -69,12 +69,12 @@ void FAnimNextControlRigTask::ExecuteControlRig(UE::AnimNext::FEvaluationVM& VM,
 			return;
 		}
 
-		const UE::AnimNext::FLODPoseStack& Pose = KeyFrameState.Pose;
-		const UE::AnimNext::FReferencePose& RefPose = Pose.GetRefPose();
+		const UE::UAF::FLODPoseStack& Pose = KeyFrameState.Pose;
+		const UE::UAF::FReferencePose& RefPose = Pose.GetRefPose();
 
 		// remap LOD pose attributes to mesh bone pose indices
 		UE::Anim::FMeshAttributeContainer MeshAttributeContainer;
-		UE::AnimNext::FGenerationTools::RemapAttributes(Pose, KeyFrameState.Attributes, MeshAttributeContainer);
+		UE::UAF::FGenerationTools::RemapAttributes(Pose, KeyFrameState.Attributes, MeshAttributeContainer);
 
 		// temporarily give control rig access to the stack allocated attribute container
 		// control rig may have rig units that can add/get attributes to/from this container
@@ -186,11 +186,11 @@ void FAnimNextControlRigTask::ExecuteControlRig(UE::AnimNext::FEvaluationVM& VM,
 		UpdateOutput(VM, KeyFrameState, ControlRig, KeyFrameState);
 		
 		// remap mesh bone index attributes back to stack container (LOD/compact bone indices)
-		UE::AnimNext::FGenerationTools::RemapAttributes(Pose, MeshAttributeContainer, KeyFrameState.Attributes);
+		UE::UAF::FGenerationTools::RemapAttributes(Pose, MeshAttributeContainer, KeyFrameState.Attributes);
 	}
 }
 
-void FAnimNextControlRigTask::UpdateInput(UE::AnimNext::FEvaluationVM& VM, UE::AnimNext::FKeyframeState& KeyFrameState, UControlRig* ControlRig, UE::AnimNext::FKeyframeState& InOutput) const
+void FAnimNextControlRigTask::UpdateInput(UE::UAF::FEvaluationVM& VM, UE::UAF::FKeyframeState& KeyFrameState, UControlRig* ControlRig, UE::UAF::FKeyframeState& InOutput) const
 {
 	DECLARE_SCOPE_HIERARCHICAL_COUNTER_FUNC()
 
@@ -215,7 +215,7 @@ void FAnimNextControlRigTask::UpdateInput(UE::AnimNext::FEvaluationVM& VM, UE::A
 	InstanceData->ControlRigVariableMappings.UpdateCurveInputs(InstanceData->ControlRig, SharedData->InputMapping, InOutput.Curves);
 }
 
-void FAnimNextControlRigTask::UpdateOutput(UE::AnimNext::FEvaluationVM& VM, UE::AnimNext::FKeyframeState& KeyFrameState, UControlRig* ControlRig, UE::AnimNext::FKeyframeState& InOutput) const
+void FAnimNextControlRigTask::UpdateOutput(UE::UAF::FEvaluationVM& VM, UE::UAF::FKeyframeState& KeyFrameState, UControlRig* ControlRig, UE::UAF::FKeyframeState& InOutput) const
 {
 	DECLARE_SCOPE_HIERARCHICAL_COUNTER_FUNC()
 
diff --git a/Engine/Plugins/Experimental/AnimNextControlRig/Source/AnimNextControlRig/Private/ControlRigTask.h b/Engine/Plugins/Experimental/AnimNextControlRig/Source/AnimNextControlRig/Private/ControlRigTask.h
--- a/Engine/Plugins/Experimental/AnimNextControlRig/Source/AnimNextControlRig/Private/ControlRigTask.h
+++ b/Engine/Plugins/Experimental/AnimNextControlRig/Source/AnimNextControlRig/Private/ControlRigTask.h
@@ -11,7 +11,7 @@
 // --- ---
 #include "ControlRigTask.generated.h"
 
-namespace UE::AnimNext
+namespace UE::UAF
 {
 struct FKeyframeState;
 }
@@ -25,20 +25,20 @@ struct ANIMNEXTCONTROLRIG_API FAnimNextControlRigTask : public FAnimNextEvaluati
 
 	DECLARE_ANIM_EVALUATION_TASK(FAnimNextControlRigTask)
 
-	static FAnimNextControlRigTask Make(const UE::AnimNext::FControlRigTrait::FSharedData* SharedData, UE::AnimNext::FControlRigTrait::FInstanceData* InstanceData);
+	static FAnimNextControlRigTask Make(const UE::UAF::FControlRigTrait::FSharedData* SharedData, UE::UAF::FControlRigTrait::FInstanceData* InstanceData);
 
 	// Task entry point
-	virtual void Execute(UE::AnimNext::FEvaluationVM& VM) const override;
+	virtual void Execute(UE::UAF::FEvaluationVM& VM) const override;
 
 private:
-	void ExecuteControlRig(UE::AnimNext::FEvaluationVM& VM, UE::AnimNext::FKeyframeState& KeyFrameState, UControlRig* ControlRig) const;
-	void UpdateInput(UE::AnimNext::FEvaluationVM& VM, UE::AnimNext::FKeyframeState& KeyFrameState, UControlRig* ControlRig, UE::AnimNext::FKeyframeState& InOutput) const;
-	void UpdateOutput(UE::AnimNext::FEvaluationVM& VM, UE::AnimNext::FKeyframeState& KeyFrameState, UControlRig* ControlRig, UE::AnimNext::FKeyframeState& InOutput) const;
+	void ExecuteControlRig(UE::UAF::FEvaluationVM& VM, UE::UAF::FKeyframeState& KeyFrameState, UControlRig* ControlRig) const;
+	void UpdateInput(UE::UAF::FEvaluationVM& VM, UE::UAF::FKeyframeState& KeyFrameState, UControlRig* ControlRig, UE::UAF::FKeyframeState& InOutput) const;
+	void UpdateOutput(UE::UAF::FEvaluationVM& VM, UE::UAF::FKeyframeState& KeyFrameState, UControlRig* ControlRig, UE::UAF::FKeyframeState& InOutput) const;
 
 	bool CanExecute(const UControlRig* ControlRig) const;
 
 	void QueueControlRigDrawInstructions(UControlRig* ControlRig, FRigVMDrawInterface* DebugDrawInterface, const FTransform& ComponentTransform) const;
 
-	const UE::AnimNext::FControlRigTrait::FSharedData* SharedData = nullptr;
-	UE::AnimNext::FControlRigTrait::FInstanceData* InstanceData = nullptr;
+	const UE::UAF::FControlRigTrait::FSharedData* SharedData = nullptr;
+	UE::UAF::FControlRigTrait::FInstanceData* InstanceData = nullptr;
 };
diff --git a/Engine/Plugins/Experimental/AnimNextControlRig/Source/AnimNextControlRig/Private/ControlRigTrait.cpp b/Engine/Plugins/Experimental/AnimNextControlRig/Source/AnimNextControlRig/Private/ControlRigTrait.cpp
--- a/Engine/Plugins/Experimental/AnimNextControlRig/Source/AnimNextControlRig/Private/ControlRigTrait.cpp
+++ b/Engine/Plugins/Experimental/AnimNextControlRig/Source/AnimNextControlRig/Private/ControlRigTrait.cpp
@@ -26,15 +26,15 @@
 //----------------------
 // --- FSharedData ---
 //----------------------
-void FControlRigTraitSharedData::ConstructLatentProperties(const UE::AnimNext::FTraitBinding& Binding)
+void FControlRigTraitSharedData::ConstructLatentProperties(const UE::UAF::FTraitBinding& Binding)
 {
 	const FControlRigTraitSharedData* SharedData = Binding.GetSharedData<FControlRigTraitSharedData>();
 	if (SharedData && SharedData->ControlRigClass.Get() != nullptr)
 	{
-		UE::AnimNext::FControlRigTrait::FInstanceData* InstanceData = Binding.GetInstanceData<UE::AnimNext::FControlRigTrait::FInstanceData>();
+		UE::UAF::FControlRigTrait::FInstanceData* InstanceData = Binding.GetInstanceData<UE::UAF::FControlRigTrait::FInstanceData>();
 
-		UE::AnimNext::FControlRigTrait::FInstanceData::GetExposedVariablesData(Binding, SharedData, InstanceData->PropertyMappings);
-		UE::AnimNext::FControlRigTrait::FInstanceData::GetExposedControlsData(Binding, SharedData, InstanceData->PropertyMappings);
+		UE::UAF::FControlRigTrait::FInstanceData::GetExposedVariablesData(Binding, SharedData, InstanceData->PropertyMappings);
+		UE::UAF::FControlRigTrait::FInstanceData::GetExposedControlsData(Binding, SharedData, InstanceData->PropertyMappings);
 
 		const TArray<FControlRigVariableMappings::FCustomPropertyData>& Mappings = InstanceData->PropertyMappings.GetMappings();
 		const int32 NumProperties = Mappings.Num();
@@ -53,12 +53,12 @@ void FControlRigTraitSharedData::ConstructLatentProperties(const UE::AnimNext::F
 	}
 }
 
-void FControlRigTraitSharedData::DestructLatentProperties(const UE::AnimNext::FTraitBinding& Binding)
+void FControlRigTraitSharedData::DestructLatentProperties(const UE::UAF::FTraitBinding& Binding)
 {
 	const FControlRigTraitSharedData* SharedData = Binding.GetSharedData<FControlRigTraitSharedData>();
 	if (SharedData && SharedData->ControlRigClass.Get() != nullptr)
 	{
-		UE::AnimNext::FControlRigTrait::FInstanceData* InstanceData = Binding.GetInstanceData<UE::AnimNext::FControlRigTrait::FInstanceData>();
+		UE::UAF::FControlRigTrait::FInstanceData* InstanceData = Binding.GetInstanceData<UE::UAF::FControlRigTrait::FInstanceData>();
 
 		const TArray<FControlRigVariableMappings::FCustomPropertyData>& Mappings = InstanceData->PropertyMappings.GetMappings();
 		const int32 NumProperties = Mappings.Num();
@@ -99,7 +99,7 @@ USkeleton* FControlRigTraitSharedData::GetPreviewSkeleton() const
 }
 #endif //WITH_EDITOR
 
-namespace UE::AnimNext
+namespace UE::UAF
 {
 
 AUTO_REGISTER_ANIM_TRAIT(FControlRigTrait)
@@ -133,7 +133,7 @@ void FControlRigTrait::FInstanceData::Construct(const FExecutionContext& Context
 
 #if WITH_EDITOR
 	{
-		OnObjectsReinstancedHandle = UE::AnimNext::ControlRig::FAnimNextControlRigModule::OnObjectsReinstanced.AddRaw(this, &FControlRigTrait::FInstanceData::OnObjectsReinstanced);
+		OnObjectsReinstancedHandle = UE::UAF::ControlRig::FAnimNextControlRigModule::OnObjectsReinstanced.AddRaw(this, &FControlRigTrait::FInstanceData::OnObjectsReinstanced);
 	}
 #endif
 
@@ -161,7 +161,7 @@ void FControlRigTrait::FInstanceData::Destruct(const FExecutionContext& Context,
 	IGarbageCollection::UnregisterWithGC(Context, Binding);
 
 #if WITH_EDITOR
-	UE::AnimNext::ControlRig::FAnimNextControlRigModule::OnObjectsReinstanced.Remove(OnObjectsReinstancedHandle);
+	UE::UAF::ControlRig::FAnimNextControlRigModule::OnObjectsReinstanced.Remove(OnObjectsReinstancedHandle);
 #endif
 
 	DestroyControlRig(Context, Binding);
@@ -239,7 +239,7 @@ UClass* FControlRigTrait::FInstanceData::GetTargetClass(const FSharedData* Share
 	return nullptr;
 }
 
-int32 FControlRigTrait::FInstanceData::GetExposedVariablesData(const UE::AnimNext::FTraitBinding& Binding, const FSharedData* SharedData, FControlRigVariableMappings::FCustomPropertyMappings& OutMappings)
+int32 FControlRigTrait::FInstanceData::GetExposedVariablesData(const UE::UAF::FTraitBinding& Binding, const FSharedData* SharedData, FControlRigVariableMappings::FCustomPropertyMappings& OutMappings)
 {
 	int32 NumElementsAdded = 0;
 
@@ -269,15 +269,15 @@ int32 FControlRigTrait::FInstanceData::GetExposedVariablesData(const UE::AnimNex
 					{
 						uint32 PropertyAlignment = 0;
 						uint32 PropertySize = 0;
-						UE::AnimNext::FControlRigTrait::GetVariableSizeAndAlignment(*Variable, PropertySize, PropertyAlignment);
+						UE::UAF::FControlRigTrait::GetVariableSizeAndAlignment(*Variable, PropertySize, PropertyAlignment);
 						check(PropertyAlignment < MAX_uint16);
 						check(PropertySize < MAX_uint16);
 
 						if (ensure(SharedData->LatentPropertyMemoryLayouts[LatentPropertyIndex] == ((PropertySize << 16) | PropertyAlignment)))
 						{
-							const UE::AnimNext::FLatentPropertyHandle* TraitLatentPropertyHandles = Binding.GetLatentPropertyHandles();
+							const UE::UAF::FLatentPropertyHandle* TraitLatentPropertyHandles = Binding.GetLatentPropertyHandles();
 
-							const UE::AnimNext::FLatentPropertyHandle& LatentPropertyHandle = TraitLatentPropertyHandles[LatentPropertyIndex];
+							const UE::UAF::FLatentPropertyHandle& LatentPropertyHandle = TraitLatentPropertyHandles[LatentPropertyIndex];
 							if (LatentPropertyHandle.IsOffsetValid())
 							{
 								const uint8* LatentPropertyMemory = const_cast<uint8*>(Binding.GetLatentProperty<uint8>(LatentPropertyHandle));
@@ -293,7 +293,7 @@ int32 FControlRigTrait::FInstanceData::GetExposedVariablesData(const UE::AnimNex
 	return NumElementsAdded;
 }
 
-int32 FControlRigTrait::FInstanceData::GetExposedControlsData(const UE::AnimNext::FTraitBinding& Binding, const FSharedData* SharedData, FControlRigVariableMappings::FCustomPropertyMappings& OutMappings)
+int32 FControlRigTrait::FInstanceData::GetExposedControlsData(const UE::UAF::FTraitBinding& Binding, const FSharedData* SharedData, FControlRigVariableMappings::FCustomPropertyMappings& OutMappings)
 {
 	int32 NumElementsAdded = 0;
 
@@ -321,15 +321,15 @@ int32 FControlRigTrait::FInstanceData::GetExposedControlsData(const UE::AnimNext
 
 					uint32 PropertyAlignment = 0;
 					uint32 PropertySize = 0;
-					UE::AnimNext::FControlRigTrait::GetControlSizeAndAlignment(ControlType, PropertySize, PropertyAlignment);
+					UE::UAF::FControlRigTrait::GetControlSizeAndAlignment(ControlType, PropertySize, PropertyAlignment);
 					ensure(PropertyAlignment < MAX_uint16);
 					ensure(PropertySize < MAX_uint16);
 
 					if (ensure(SharedData->LatentPropertyMemoryLayouts[LatentPropertyIndex] == ((PropertySize << 16) | PropertyAlignment)))
 					{
-						const UE::AnimNext::FLatentPropertyHandle* TraitLatentPropertyHandles = Binding.GetLatentPropertyHandles();
+						const UE::UAF::FLatentPropertyHandle* TraitLatentPropertyHandles = Binding.GetLatentPropertyHandles();
 
-						const UE::AnimNext::FLatentPropertyHandle& LatentPropertyHandle = TraitLatentPropertyHandles[LatentPropertyIndex];
+						const UE::UAF::FLatentPropertyHandle& LatentPropertyHandle = TraitLatentPropertyHandles[LatentPropertyIndex];
 						if (LatentPropertyHandle.IsOffsetValid())
 						{
 							const uint8* LatentPropertyMemory = const_cast<uint8*>(Binding.GetLatentProperty<uint8>(LatentPropertyHandle));
@@ -354,7 +354,7 @@ void FControlRigTrait::FInstanceData::HandleOnInitialized_AnyThread(URigVMHost*,
 
 void FControlRigTrait::FInstanceData::InitializeCustomProperties(const FTraitBinding& Binding, const FSharedData* SharedData)
 {
-	UE::AnimNext::FControlRigTrait::FInstanceData* InstanceData = Binding.GetInstanceData<UE::AnimNext::FControlRigTrait::FInstanceData>();
+	UE::UAF::FControlRigTrait::FInstanceData* InstanceData = Binding.GetInstanceData<UE::UAF::FControlRigTrait::FInstanceData>();
 	 
 	// Setup mappings using the latent pin memory as source (we have to copy from latent pin to external variable / rig control)
 	ControlRigVariableMappings.InitializeCustomProperties(ControlRig, InstanceData->PropertyMappings);
@@ -530,8 +530,8 @@ void FControlRigTrait::PreUpdate(FUpdateTraversalContext& Context, const TTraitB
 	if (InstanceData->bReinitializeControlRig)
 	{
 		InstanceData->PropertyMappings.Reset();
-		UE::AnimNext::FControlRigTrait::FInstanceData::GetExposedVariablesData(Binding, SharedData, InstanceData->PropertyMappings);
-		UE::AnimNext::FControlRigTrait::FInstanceData::GetExposedControlsData(Binding, SharedData, InstanceData->PropertyMappings);
+		UE::UAF::FControlRigTrait::FInstanceData::GetExposedVariablesData(Binding, SharedData, InstanceData->PropertyMappings);
+		UE::UAF::FControlRigTrait::FInstanceData::GetExposedControlsData(Binding, SharedData, InstanceData->PropertyMappings);
 
 		InstanceData->InitializeControlRig(Context, Binding);
 		InstanceData->bReinitializeControlRig = false;
@@ -958,4 +958,4 @@ bool FControlRigTrait::GetControlSizeAndAlignment(ERigControlType ControlType, u
 	return false;
 }
 
-} // end namespace UE::AnimNext
+} // end namespace UE::UAF
diff --git a/Engine/Plugins/Experimental/AnimNextControlRig/Source/AnimNextControlRigEditor/Private/AnimNextControlRigEditorModule.cpp b/Engine/Plugins/Experimental/AnimNextControlRig/Source/AnimNextControlRigEditor/Private/AnimNextControlRigEditorModule.cpp
--- a/Engine/Plugins/Experimental/AnimNextControlRig/Source/AnimNextControlRigEditor/Private/AnimNextControlRigEditorModule.cpp
+++ b/Engine/Plugins/Experimental/AnimNextControlRig/Source/AnimNextControlRigEditor/Private/AnimNextControlRigEditorModule.cpp
@@ -18,7 +18,7 @@
 
 extern UNREALED_API UEditorEngine* GEditor;
 
-namespace UE::AnimNext::ControlRig::Editor
+namespace UE::UAF::ControlRig::Editor
 {
 
 void FAnimNextControlRigEditorModule::StartupModule()
@@ -45,11 +45,11 @@ void FAnimNextControlRigEditorModule::StartupModule()
 	// This generates all the Details panel ControlRig specific data and enables programmatic pin generation
 	PropertyModule.RegisterCustomPropertyTypeLayout(
 		"ControlRigTraitSharedData",
-		FOnGetPropertyTypeCustomizationInstance::CreateLambda([] { return MakeShared<UE::AnimNext::Editor::FControlRigTraitSharedDataCustomization>(); }));
+		FOnGetPropertyTypeCustomizationInstance::CreateLambda([] { return MakeShared<UE::UAF::Editor::FControlRigTraitSharedDataCustomization>(); }));
 
-	UE::AnimNext::Editor::IAnimNextEditorModule& AnimNextEditorModule = FModuleManager::LoadModuleChecked<UE::AnimNext::Editor::IAnimNextEditorModule>("AnimNextEditor");
+	UE::UAF::Editor::IAnimNextEditorModule& AnimNextEditorModule = FModuleManager::LoadModuleChecked<UE::UAF::Editor::IAnimNextEditorModule>("AnimNextEditor");
 
-	OnNodeDblClickDelegateHandle = AnimNextEditorModule.RegisterNodeDblClickHandler(UE::AnimNext::Editor::IAnimNextEditorModule::FNodeDblClickNotificationDelegate::CreateLambda([](const Workspace::FWorkspaceEditorContext& InContext, const UEdGraphNode* InNode)
+	OnNodeDblClickDelegateHandle = AnimNextEditorModule.RegisterNodeDblClickHandler(UE::UAF::Editor::IAnimNextEditorModule::FNodeDblClickNotificationDelegate::CreateLambda([](const Workspace::FWorkspaceEditorContext& InContext, const UEdGraphNode* InNode)
 		{
 			if (const UAnimNextEdGraphNode* RigVMEdGraphNode = Cast<UAnimNextEdGraphNode>(InNode))
 			{
@@ -91,7 +91,7 @@ void FAnimNextControlRigEditorModule::ShutdownModule()
 {
 	if (FModuleManager::Get().IsModuleLoaded("AnimNextEditor"))
 	{
-		UE::AnimNext::Editor::IAnimNextEditorModule& AnimNextEditorModule = FModuleManager::LoadModuleChecked<UE::AnimNext::Editor::IAnimNextEditorModule>("AnimNextEditor");
+		UE::UAF::Editor::IAnimNextEditorModule& AnimNextEditorModule = FModuleManager::LoadModuleChecked<UE::UAF::Editor::IAnimNextEditorModule>("AnimNextEditor");
 
 		AnimNextEditorModule.UnregisterNodeDblClickHandler(OnNodeDblClickDelegateHandle);
 	}
@@ -99,4 +99,4 @@ void FAnimNextControlRigEditorModule::ShutdownModule()
 
 } // end namespace
 
-IMPLEMENT_MODULE(UE::AnimNext::ControlRig::Editor::FAnimNextControlRigEditorModule, AnimNextControlRigEditor)
+IMPLEMENT_MODULE(UE::UAF::ControlRig::Editor::FAnimNextControlRigEditorModule, AnimNextControlRigEditor)
diff --git a/Engine/Plugins/Experimental/AnimNextControlRig/Source/AnimNextControlRigEditor/Private/AnimNextControlRigEditorModule.h b/Engine/Plugins/Experimental/AnimNextControlRig/Source/AnimNextControlRigEditor/Private/AnimNextControlRigEditorModule.h
--- a/Engine/Plugins/Experimental/AnimNextControlRig/Source/AnimNextControlRigEditor/Private/AnimNextControlRigEditorModule.h
+++ b/Engine/Plugins/Experimental/AnimNextControlRig/Source/AnimNextControlRigEditor/Private/AnimNextControlRigEditorModule.h
@@ -4,7 +4,7 @@
 #include "CoreMinimal.h"
 #include "Modules/ModuleInterface.h"
 
-namespace UE::AnimNext::ControlRig::Editor
+namespace UE::UAF::ControlRig::Editor
 {
 
 class FAnimNextControlRigEditorModule : public IModuleInterface
diff --git a/Engine/Plugins/Experimental/AnimNextControlRig/Source/AnimNextControlRigEditor/Private/ControlRigTraitCustomization.cpp b/Engine/Plugins/Experimental/AnimNextControlRig/Source/AnimNextControlRigEditor/Private/ControlRigTraitCustomization.cpp
--- a/Engine/Plugins/Experimental/AnimNextControlRig/Source/AnimNextControlRigEditor/Private/ControlRigTraitCustomization.cpp
+++ b/Engine/Plugins/Experimental/AnimNextControlRig/Source/AnimNextControlRigEditor/Private/ControlRigTraitCustomization.cpp
@@ -25,14 +25,14 @@
 
 #define LOCTEXT_NAMESPACE "ControlRigTraitSharedDataCustomization"
 
-namespace UE::AnimNext::Editor
+namespace UE::UAF::Editor
 {
 
 FControlRigTraitSharedDataCustomization::~FControlRigTraitSharedDataCustomization()
 {
 	if (OnObjectsReinstancedHandle.IsValid())
 	{
-		UE::AnimNext::ControlRig::FAnimNextControlRigModule::OnObjectsReinstanced.Remove(OnObjectsReinstancedHandle);
+		UE::UAF::ControlRig::FAnimNextControlRigModule::OnObjectsReinstanced.Remove(OnObjectsReinstancedHandle);
 	}
 }
 
@@ -74,7 +74,7 @@ void FControlRigTraitSharedDataCustomization::CustomizeHeader(TSharedRef<IProper
 	ControlRigIOMapping->GetOnGetTargetClassDelegate().BindSP(this, &FControlRigTraitSharedDataCustomization::GetTargetClass);
 	ControlRigIOMapping->SetIgnoreVariablesWithNoMemory(true); // AnimBP creates it's own variables, AnimNext needs the memory to create the pin
 
-	OnObjectsReinstancedHandle = UE::AnimNext::ControlRig::FAnimNextControlRigModule::OnObjectsReinstanced.AddRaw(this, &FControlRigTraitSharedDataCustomization::OnObjectsReinstanced);
+	OnObjectsReinstancedHandle = UE::UAF::ControlRig::FAnimNextControlRigModule::OnObjectsReinstanced.AddRaw(this, &FControlRigTraitSharedDataCustomization::OnObjectsReinstanced);
 
 	InHeaderRow.NameContent()
 	[
@@ -435,7 +435,7 @@ FString FControlRigTraitSharedDataCustomization::FindControlRigTraitPinName(URig
 	return FString();
 }
 
-} // end namespace UE::AnimNext::Editor
+} // end namespace UE::UAF::Editor
 
 
 #undef LOCTEXT_NAMESPACE
diff --git a/Engine/Plugins/Experimental/AnimNextPoseSearch/Source/AnimNextPoseSearch/Private/AnimNextInteractionIslandDependency.h b/Engine/Plugins/Experimental/AnimNextPoseSearch/Source/AnimNextPoseSearch/Private/AnimNextInteractionIslandDependency.h
--- a/Engine/Plugins/Experimental/AnimNextPoseSearch/Source/AnimNextPoseSearch/Private/AnimNextInteractionIslandDependency.h
+++ b/Engine/Plugins/Experimental/AnimNextPoseSearch/Source/AnimNextPoseSearch/Private/AnimNextInteractionIslandDependency.h
@@ -4,7 +4,7 @@
 
 #include "PoseSearch/PoseSearchInteractionIsland.h"
 
-namespace UE::AnimNext::PoseSearch
+namespace UE::UAF::PoseSearch
 {
 	class FModule;
 }
@@ -15,7 +15,7 @@ namespace UE::PoseSearch
 struct FAnimNextInteractionIslandDependency : public IInteractionIslandDependency
 {
 private:
-	friend UE::AnimNext::PoseSearch::FModule;
+	friend UE::UAF::PoseSearch::FModule;
 	
 	virtual bool CanMakeDependency(const UObject* InIslandObject, const UObject* InObject) const override;
 	virtual const FTickFunction* FindTickFunction(UObject* InObject) const override;
diff --git a/Engine/Plugins/Experimental/AnimNextPoseSearch/Source/AnimNextPoseSearch/Private/HistoryCollectorTrait.cpp b/Engine/Plugins/Experimental/AnimNextPoseSearch/Source/AnimNextPoseSearch/Private/HistoryCollectorTrait.cpp
--- a/Engine/Plugins/Experimental/AnimNextPoseSearch/Source/AnimNextPoseSearch/Private/HistoryCollectorTrait.cpp
+++ b/Engine/Plugins/Experimental/AnimNextPoseSearch/Source/AnimNextPoseSearch/Private/HistoryCollectorTrait.cpp
@@ -10,7 +10,7 @@
 #include "TraitCore/NodeInstance.h"
 #include "PoseHistoryEvaluation.h"
 
-namespace UE::AnimNext
+namespace UE::UAF
 {
 	// @todo: revisit this code once a FCSPose<FLODPoseStack> is implemented!
 	struct FComponentSpacePoseProvider : public UE::PoseSearch::IComponentSpacePoseProvider
@@ -80,9 +80,9 @@ namespace UE::AnimNext
 	};
 }
 
-void FAnimNextHistoryCollectorPreEvaluateTask::Execute(UE::AnimNext::FEvaluationVM& VM) const
+void FAnimNextHistoryCollectorPreEvaluateTask::Execute(UE::UAF::FEvaluationVM& VM) const
 {
-	using namespace UE::AnimNext;
+	using namespace UE::UAF;
 	using namespace UE::PoseSearch;
 
 	if (EnumHasAnyFlags(VM.GetFlags(), EEvaluationFlags::Bones))
@@ -93,9 +93,9 @@ void FAnimNextHistoryCollectorPreEvaluateTask::Execute(UE::AnimNext::FEvaluation
 	}
 }
 
-void FAnimNextHistoryCollectorTask::Execute(UE::AnimNext::FEvaluationVM& VM) const
+void FAnimNextHistoryCollectorTask::Execute(UE::UAF::FEvaluationVM& VM) const
 {
-	using namespace UE::AnimNext;
+	using namespace UE::UAF;
 
 	if (EnumHasAnyFlags(VM.GetFlags(), EEvaluationFlags::Bones))
 	{
@@ -148,7 +148,7 @@ void FAnimNextHistoryCollectorTask::Execute(UE::AnimNext::FEvaluationVM& VM) con
 	}
 }
 
-namespace UE::AnimNext
+namespace UE::UAF
 {
 	AUTO_REGISTER_ANIM_TRAIT(FHistoryCollectorTrait)
 
diff --git a/Engine/Plugins/Experimental/AnimNextPoseSearch/Source/AnimNextPoseSearch/Private/HistoryCollectorTrait.h b/Engine/Plugins/Experimental/AnimNextPoseSearch/Source/AnimNextPoseSearch/Private/HistoryCollectorTrait.h
--- a/Engine/Plugins/Experimental/AnimNextPoseSearch/Source/AnimNextPoseSearch/Private/HistoryCollectorTrait.h
+++ b/Engine/Plugins/Experimental/AnimNextPoseSearch/Source/AnimNextPoseSearch/Private/HistoryCollectorTrait.h
@@ -121,7 +121,7 @@ struct FAnimNextHistoryCollectorTask : public FAnimNextEvaluationTask
 
 	DECLARE_ANIM_EVALUATION_TASK(FAnimNextHistoryCollectorTask)
 
-	virtual void Execute(UE::AnimNext::FEvaluationVM& VM) const override;
+	virtual void Execute(UE::UAF::FEvaluationVM& VM) const override;
 
 	UE::PoseSearch::FPoseHistory* PoseHistory = nullptr;
 	const FAnimNextHistoryCollectorTraitSharedData* SharedData = nullptr;
@@ -130,7 +130,7 @@ struct FAnimNextHistoryCollectorTask : public FAnimNextEvaluationTask
 	TWeakObjectPtr<const UObject> HostObject;
 };
 
-namespace UE::AnimNext
+namespace UE::UAF
 {
 	struct FHistoryCollectorTrait : FAdditiveTrait, IUpdate, IEvaluate, IPoseHistory
 	{
@@ -172,6 +172,6 @@ struct FAnimNextHistoryCollectorPreEvaluateTask : public FAnimNextEvaluationTask
 
 	DECLARE_ANIM_EVALUATION_TASK(FAnimNextHistoryCollectorPreEvaluateTask)
 
-	virtual void Execute(UE::AnimNext::FEvaluationVM& VM) const override;
-	UE::AnimNext::FHistoryCollectorTrait::FInstanceData* InstanceData = nullptr;
+	virtual void Execute(UE::UAF::FEvaluationVM& VM) const override;
+	UE::UAF::FHistoryCollectorTrait::FInstanceData* InstanceData = nullptr;
 };
\ No newline at end of file
diff --git a/Engine/Plugins/Experimental/AnimNextPoseSearch/Source/AnimNextPoseSearch/Private/Module.cpp b/Engine/Plugins/Experimental/AnimNextPoseSearch/Source/AnimNextPoseSearch/Private/Module.cpp
--- a/Engine/Plugins/Experimental/AnimNextPoseSearch/Source/AnimNextPoseSearch/Private/Module.cpp
+++ b/Engine/Plugins/Experimental/AnimNextPoseSearch/Source/AnimNextPoseSearch/Private/Module.cpp
@@ -15,7 +15,7 @@
 #include "PoseSearch/PoseSearchSchema.h"
 #include "Animation/TrajectoryTypes.h"
 
-namespace UE::AnimNext::PoseSearch
+namespace UE::UAF::PoseSearch
 {
 
 class FModule : public IModuleInterface
@@ -57,4 +57,4 @@ public:
 
 }
 
-IMPLEMENT_MODULE(UE::AnimNext::PoseSearch::FModule, AnimNextPoseSearch)
+IMPLEMENT_MODULE(UE::UAF::PoseSearch::FModule, AnimNextPoseSearch)
diff --git a/Engine/Plugins/Experimental/AnimNextPoseSearch/Source/AnimNextPoseSearch/Private/MotionMatchingTrait.cpp b/Engine/Plugins/Experimental/AnimNextPoseSearch/Source/AnimNextPoseSearch/Private/MotionMatchingTrait.cpp
--- a/Engine/Plugins/Experimental/AnimNextPoseSearch/Source/AnimNextPoseSearch/Private/MotionMatchingTrait.cpp
+++ b/Engine/Plugins/Experimental/AnimNextPoseSearch/Source/AnimNextPoseSearch/Private/MotionMatchingTrait.cpp
@@ -48,7 +48,7 @@ static FAutoConsoleVariableRef CVarAnimNextMotionMatchingPlayRateEnabled(
 );
 #endif // ENABLE_ANIM_DEBUG
 
-namespace UE::AnimNext
+namespace UE::UAF
 {
 	AUTO_REGISTER_ANIM_TRAIT(FMotionMatchingTrait)
 
@@ -490,7 +490,7 @@ namespace UE::AnimNext
 							FRigUnit_AnimNextPrePhysicsEvent::DefaultEventName,
 							[SelectedAnim, SelectedTime, bIsContinuingPoseSearch, WantedPlayRate, bLoop,
 							bIsMirrored, BlendParameters, SelectedDatabase, SearchCost, bIsInteraction,
-							Role, ActorRootTransforms, ActorRootBoneTransforms, ResultOutputVariableName](const UE::AnimNext::FModuleTaskContext& TaskContext)
+							Role, ActorRootTransforms, ActorRootBoneTransforms, ResultOutputVariableName](const UE::UAF::FModuleTaskContext& TaskContext)
 							{
 								FAnimNextModuleInstance* OtherModuleInstance = TaskContext.GetModuleInstance();
 								if (OtherModuleInstance)
@@ -515,7 +515,7 @@ namespace UE::AnimNext
 									PRAGMA_ENABLE_DEPRECATION_WARNINGS
 								}
 							},
-							UE::AnimNext::ETaskRunLocation::Before);
+							UE::UAF::ETaskRunLocation::Before);
 					}
 					else
 					{
@@ -617,11 +617,11 @@ namespace UE::AnimNext
 		}
 	}
 
-} // UE::AnimNext
+} // UE::UAF
 
-void FAnimNextMotionMatchingTask::Execute(UE::AnimNext::FEvaluationVM& VM) const
+void FAnimNextMotionMatchingTask::Execute(UE::UAF::FEvaluationVM& VM) const
 {
-	using namespace UE::AnimNext;
+	using namespace UE::UAF;
 	using namespace UE::PoseSearch;
 
 	check(CurrentResultRoleIndex != INDEX_NONE);
diff --git a/Engine/Plugins/Experimental/AnimNextPoseSearch/Source/AnimNextPoseSearch/Private/MotionMatchingTrait.h b/Engine/Plugins/Experimental/AnimNextPoseSearch/Source/AnimNextPoseSearch/Private/MotionMatchingTrait.h
--- a/Engine/Plugins/Experimental/AnimNextPoseSearch/Source/AnimNextPoseSearch/Private/MotionMatchingTrait.h
+++ b/Engine/Plugins/Experimental/AnimNextPoseSearch/Source/AnimNextPoseSearch/Private/MotionMatchingTrait.h
@@ -158,7 +158,7 @@ struct FMotionMatchingTraitSharedData : public FAnimNextTraitSharedData
 	#undef TRAIT_LATENT_PROPERTIES_ENUMERATOR
 };
 
-namespace UE::AnimNext
+namespace UE::UAF
 {
 
 struct FMotionMatchingTrait : FAdditiveTrait, IUpdate, IEvaluate, IGarbageCollection
@@ -193,7 +193,7 @@ private:
 	void PublishResults(const TTraitBinding<IUpdate>& Binding) const;
 };
 
-} // namespace UE::AnimNext
+} // namespace UE::UAF
 
 USTRUCT(Experimental)
 struct FAnimNextMotionMatchingTask : public FAnimNextEvaluationTask
@@ -202,10 +202,10 @@ struct FAnimNextMotionMatchingTask : public FAnimNextEvaluationTask
 
 	DECLARE_ANIM_EVALUATION_TASK(FAnimNextMotionMatchingTask)
 
-	virtual void Execute(UE::AnimNext::FEvaluationVM& VM) const override;
+	virtual void Execute(UE::UAF::FEvaluationVM& VM) const override;
 
 	FTransform ComponentTransform = FTransform::Identity;
-	UE::AnimNext::FMotionMatchingTrait::FInstanceData* InstanceData = nullptr;
+	UE::UAF::FMotionMatchingTrait::FInstanceData* InstanceData = nullptr;
 	int32 CurrentResultRoleIndex = INDEX_NONE;
 	bool bWarpUsingRootBone = true;
 	float WarpingRotationRatio = 1.f;
diff --git a/Engine/Plugins/Experimental/AnimNextPoseSearch/Source/AnimNextPoseSearch/Private/PoseHistoryEvaluation.cpp b/Engine/Plugins/Experimental/AnimNextPoseSearch/Source/AnimNextPoseSearch/Private/PoseHistoryEvaluation.cpp
--- a/Engine/Plugins/Experimental/AnimNextPoseSearch/Source/AnimNextPoseSearch/Private/PoseHistoryEvaluation.cpp
+++ b/Engine/Plugins/Experimental/AnimNextPoseSearch/Source/AnimNextPoseSearch/Private/PoseHistoryEvaluation.cpp
@@ -4,5 +4,5 @@
 
 namespace UE::PoseSearch
 {
-	const UE::AnimNext::FEvaluationVMStackName POSEHISTORY_STACK_NAME = FName(TEXT("PoseHistoryStack"));
+	const UE::UAF::FEvaluationVMStackName POSEHISTORY_STACK_NAME = FName(TEXT("PoseHistoryStack"));
 }
\ No newline at end of file
diff --git a/Engine/Plugins/Experimental/AnimNextPoseSearch/Source/AnimNextPoseSearch/Private/PoseSearchResultEmulatorTrait.cpp b/Engine/Plugins/Experimental/AnimNextPoseSearch/Source/AnimNextPoseSearch/Private/PoseSearchResultEmulatorTrait.cpp
--- a/Engine/Plugins/Experimental/AnimNextPoseSearch/Source/AnimNextPoseSearch/Private/PoseSearchResultEmulatorTrait.cpp
+++ b/Engine/Plugins/Experimental/AnimNextPoseSearch/Source/AnimNextPoseSearch/Private/PoseSearchResultEmulatorTrait.cpp
@@ -15,7 +15,7 @@
 #include "Factory/AnimGraphFactory.h"
 #include "Traits/SynchronizeUsingGroupsTraitData.h"
 
-namespace UE::AnimNext
+namespace UE::UAF
 {
 	AUTO_REGISTER_ANIM_TRAIT(FPoseSearchResultEmulatorTrait)
 
@@ -187,11 +187,11 @@ namespace UE::AnimNext
 		}
 	}
 	
-} // UE::AnimNext
+} // UE::UAF
 
-void FAnimNextPoseSearchResultEmulatorTask::Execute(UE::AnimNext::FEvaluationVM& VM) const
+void FAnimNextPoseSearchResultEmulatorTask::Execute(UE::UAF::FEvaluationVM& VM) const
 {
-	using namespace UE::AnimNext;
+	using namespace UE::UAF;
 
 	const UE::Anim::IAnimRootMotionProvider* RootMotionProvider = UE::Anim::IAnimRootMotionProvider::Get();
 	if (!RootMotionProvider)
diff --git a/Engine/Plugins/Experimental/AnimNextPoseSearch/Source/AnimNextPoseSearch/Private/PoseSearchResultEmulatorTrait.h b/Engine/Plugins/Experimental/AnimNextPoseSearch/Source/AnimNextPoseSearch/Private/PoseSearchResultEmulatorTrait.h
--- a/Engine/Plugins/Experimental/AnimNextPoseSearch/Source/AnimNextPoseSearch/Private/PoseSearchResultEmulatorTrait.h
+++ b/Engine/Plugins/Experimental/AnimNextPoseSearch/Source/AnimNextPoseSearch/Private/PoseSearchResultEmulatorTrait.h
@@ -68,7 +68,7 @@ struct FPoseSearchResultEmulatorTraitSharedData : public FAnimNextTraitSharedDat
 	#undef TRAIT_LATENT_PROPERTIES_ENUMERATOR
 };
 
-namespace UE::AnimNext
+namespace UE::UAF
 {
 
 struct FPoseSearchResultEmulatorTrait : FAdditiveTrait, IUpdate, IEvaluate
@@ -88,7 +88,7 @@ struct FPoseSearchResultEmulatorTrait : FAdditiveTrait, IUpdate, IEvaluate
 	virtual void PostEvaluate(FEvaluateTraversalContext& Context, const TTraitBinding<IEvaluate>& Binding) const override;
 };
 
-} // namespace UE::AnimNext
+} // namespace UE::UAF
 
 USTRUCT(Experimental)
 struct FAnimNextPoseSearchResultEmulatorTask : public FAnimNextEvaluationTask
@@ -97,5 +97,5 @@ struct FAnimNextPoseSearchResultEmulatorTask : public FAnimNextEvaluationTask
 
 	DECLARE_ANIM_EVALUATION_TASK(FAnimNextPoseSearchResultEmulatorTask)
 
-	virtual void Execute(UE::AnimNext::FEvaluationVM& VM) const override;
+	virtual void Execute(UE::UAF::FEvaluationVM& VM) const override;
 };
\ No newline at end of file
diff --git a/Engine/Plugins/Experimental/AnimNextPoseSearch/Source/AnimNextPoseSearch/Public/PoseHistoryEvaluation.h b/Engine/Plugins/Experimental/AnimNextPoseSearch/Source/AnimNextPoseSearch/Public/PoseHistoryEvaluation.h
--- a/Engine/Plugins/Experimental/AnimNextPoseSearch/Source/AnimNextPoseSearch/Public/PoseHistoryEvaluation.h
+++ b/Engine/Plugins/Experimental/AnimNextPoseSearch/Source/AnimNextPoseSearch/Public/PoseHistoryEvaluation.h
@@ -7,7 +7,7 @@
 
 namespace UE::PoseSearch
 {
-	extern ANIMNEXTPOSESEARCH_API const UE::AnimNext::FEvaluationVMStackName POSEHISTORY_STACK_NAME;
+	extern ANIMNEXTPOSESEARCH_API const UE::UAF::FEvaluationVMStackName POSEHISTORY_STACK_NAME;
 
 	struct FPoseHistoryEvaluationHelper
 	{
@@ -15,7 +15,7 @@ namespace UE::PoseSearch
 	};
 }
 
-namespace UE::AnimNext
+namespace UE::UAF
 {
 	ANIM_NEXT_ENABLE_EVALUATION_STACK_USAGE(TUniquePtr<UE::PoseSearch::FPoseHistoryEvaluationHelper>);
 }
\ No newline at end of file
diff --git a/Engine/Plugins/Experimental/AnimNextStateTree/Source/AnimNextStateTree/Internal/Conditions/AnimNextStateTreeRigVMConditionBase.h b/Engine/Plugins/Experimental/AnimNextStateTree/Source/AnimNextStateTree/Internal/Conditions/AnimNextStateTreeRigVMConditionBase.h
--- a/Engine/Plugins/Experimental/AnimNextStateTree/Source/AnimNextStateTree/Internal/Conditions/AnimNextStateTreeRigVMConditionBase.h
+++ b/Engine/Plugins/Experimental/AnimNextStateTree/Source/AnimNextStateTree/Internal/Conditions/AnimNextStateTreeRigVMConditionBase.h
@@ -31,7 +31,7 @@ struct FAnimNextStateTreeRigVMConditionInstanceData
 	FInstancedPropertyBag Parameters;
 
 	/** Cached function handle */
-	UE::AnimNext::FFunctionHandle FunctionHandle;
+	UE::UAF::FFunctionHandle FunctionHandle;
 };
 
 /**
diff --git a/Engine/Plugins/Experimental/AnimNextStateTree/Source/AnimNextStateTree/Internal/IAnimNextStateTreeModule.h b/Engine/Plugins/Experimental/AnimNextStateTree/Source/AnimNextStateTree/Internal/IAnimNextStateTreeModule.h
--- a/Engine/Plugins/Experimental/AnimNextStateTree/Source/AnimNextStateTree/Internal/IAnimNextStateTreeModule.h
+++ b/Engine/Plugins/Experimental/AnimNextStateTree/Source/AnimNextStateTree/Internal/IAnimNextStateTreeModule.h
@@ -5,7 +5,7 @@
 #include "CoreMinimal.h"
 #include "Modules/ModuleManager.h"
 
-namespace UE::AnimNext::StateTree
+namespace UE::UAF::StateTree
 {
 
 class IAnimNextStateTreeModule : public IModuleInterface
diff --git a/Engine/Plugins/Experimental/AnimNextStateTree/Source/AnimNextStateTree/Internal/Tasks/AnimNextStateTreeRigVMTaskBase.h b/Engine/Plugins/Experimental/AnimNextStateTree/Source/AnimNextStateTree/Internal/Tasks/AnimNextStateTreeRigVMTaskBase.h
--- a/Engine/Plugins/Experimental/AnimNextStateTree/Source/AnimNextStateTree/Internal/Tasks/AnimNextStateTreeRigVMTaskBase.h
+++ b/Engine/Plugins/Experimental/AnimNextStateTree/Source/AnimNextStateTree/Internal/Tasks/AnimNextStateTreeRigVMTaskBase.h
@@ -30,7 +30,7 @@ struct ANIMNEXTSTATETREE_API FAnimNextStateTreeRigVMTaskInstanceData
 	FInstancedPropertyBag Parameters;
 
 	/** Cached function handle */
-	UE::AnimNext::FFunctionHandle FunctionHandle;
+	UE::UAF::FFunctionHandle FunctionHandle;
 };
 
 /**
diff --git a/Engine/Plugins/Experimental/AnimNextStateTree/Source/AnimNextStateTree/Private/AnimNextStateTree.cpp b/Engine/Plugins/Experimental/AnimNextStateTree/Source/AnimNextStateTree/Private/AnimNextStateTree.cpp
--- a/Engine/Plugins/Experimental/AnimNextStateTree/Source/AnimNextStateTree/Private/AnimNextStateTree.cpp
+++ b/Engine/Plugins/Experimental/AnimNextStateTree/Source/AnimNextStateTree/Private/AnimNextStateTree.cpp
@@ -5,7 +5,7 @@
 #include "Serialization/CustomVersion.h"
 
 const FGuid FAnimNextStateTreeCustomVersion::GUID(0x45641511, 0x102F42BB, 0xA6EF181D, 0x6C442CAC);
-namespace UE::AnimNextStateTree::Private
+namespace UE::UAF::StateTree::Private
 {
 FCustomVersionRegistration GRegisterStateTreeCustomVersion(FAnimNextStateTreeCustomVersion::GUID, FAnimNextStateTreeCustomVersion::LatestVersion, TEXT("AnimNextStateTree"));
 }
diff --git a/Engine/Plugins/Experimental/AnimNextStateTree/Source/AnimNextStateTree/Private/AnimNextStateTreeContext.cpp b/Engine/Plugins/Experimental/AnimNextStateTree/Source/AnimNextStateTree/Private/AnimNextStateTreeContext.cpp
--- a/Engine/Plugins/Experimental/AnimNextStateTree/Source/AnimNextStateTree/Private/AnimNextStateTreeContext.cpp
+++ b/Engine/Plugins/Experimental/AnimNextStateTree/Source/AnimNextStateTree/Private/AnimNextStateTreeContext.cpp
@@ -9,7 +9,7 @@
 
 bool FAnimNextStateTreeTraitContext::PushAssetOntoBlendStack(TNonNullPtr<UObject> InAsset, const FAlphaBlendArgs& InBlendArguments, const FAnimNextFactoryParams& InFactoryParams) const
 {
-	using namespace UE::AnimNext;
+	using namespace UE::UAF;
 
 	TTraitBinding<IBlendStack> BlendStackBinding;
 	if(!Binding->GetStackInterface<IBlendStack>(BlendStackBinding))
@@ -37,8 +37,8 @@ bool FAnimNextStateTreeTraitContext::PushAssetOntoBlendStack(TNonNullPtr<UObject
 
 float FAnimNextStateTreeTraitContext::QueryPlaybackRatio(TNonNullPtr<UObject> InAsset) const
 {
-	UE::AnimNext::TTraitBinding<UE::AnimNext::ITimeline> Timeline;
-	if (Binding->GetStackInterface<UE::AnimNext::ITimeline>(Timeline))
+	UE::UAF::TTraitBinding<UE::UAF::ITimeline> Timeline;
+	if (Binding->GetStackInterface<UE::UAF::ITimeline>(Timeline))
 	{
 		return Timeline.GetState(*Context).GetPositionRatio();
 	}
diff --git a/Engine/Plugins/Experimental/AnimNextStateTree/Source/AnimNextStateTree/Private/AnimNextStateTreeContext.h b/Engine/Plugins/Experimental/AnimNextStateTree/Source/AnimNextStateTree/Private/AnimNextStateTreeContext.h
--- a/Engine/Plugins/Experimental/AnimNextStateTree/Source/AnimNextStateTree/Private/AnimNextStateTreeContext.h
+++ b/Engine/Plugins/Experimental/AnimNextStateTree/Source/AnimNextStateTree/Private/AnimNextStateTreeContext.h
@@ -7,7 +7,7 @@
 
 #include "AnimNextStateTreeContext.generated.h"
 
-namespace UE::AnimNext
+namespace UE::UAF
 {
 struct FTraitBinding;
 struct FExecutionContext;
@@ -23,12 +23,12 @@ struct FAnimNextStateTreeTraitContext
 {	
 	GENERATED_BODY()
 	
-	friend UE::AnimNext::FStateTreeTrait;
-	friend UE::AnimNext::FStateTreeTrait::FInstanceData;
+	friend UE::UAF::FStateTreeTrait;
+	friend UE::UAF::FStateTreeTrait::FInstanceData;
 
 	FAnimNextStateTreeTraitContext() = default;
 
-	FAnimNextStateTreeTraitContext(UE::AnimNext::FExecutionContext& InContext, const UE::AnimNext::FTraitBinding& InBinding)
+	FAnimNextStateTreeTraitContext(UE::UAF::FExecutionContext& InContext, const UE::UAF::FTraitBinding& InBinding)
 		: Context(&InContext)
 		, Binding(&InBinding)
 	{}
@@ -37,20 +37,20 @@ struct FAnimNextStateTreeTraitContext
 	float QueryPlaybackRatio(TNonNullPtr<UObject> InAsset) const;
 	FAnimNextGraphInstance& GetGraphInstance() const;
 
-	UE::AnimNext::FExecutionContext& GetContext() const
+	UE::UAF::FExecutionContext& GetContext() const
 	{
 		check(Context != nullptr);
 		return *Context;
 	}
 
-	const UE::AnimNext::FTraitBinding& GetBinding() const
+	const UE::UAF::FTraitBinding& GetBinding() const
 	{
 		check(Binding != nullptr);
 		return *Binding;
 	}
 
 private:
-	UE::AnimNext::FExecutionContext* Context = nullptr;
-	const UE::AnimNext::FTraitBinding* Binding = nullptr;
+	UE::UAF::FExecutionContext* Context = nullptr;
+	const UE::UAF::FTraitBinding* Binding = nullptr;
 };
 
diff --git a/Engine/Plugins/Experimental/AnimNextStateTree/Source/AnimNextStateTree/Private/AnimNextStateTreeModule.cpp b/Engine/Plugins/Experimental/AnimNextStateTree/Source/AnimNextStateTree/Private/AnimNextStateTreeModule.cpp
--- a/Engine/Plugins/Experimental/AnimNextStateTree/Source/AnimNextStateTree/Private/AnimNextStateTreeModule.cpp
+++ b/Engine/Plugins/Experimental/AnimNextStateTree/Source/AnimNextStateTree/Private/AnimNextStateTreeModule.cpp
@@ -7,7 +7,7 @@
 #include "Graph/AnimNextAnimationGraph.h"
 #include "RigVMCore/RigVMRegistry.h"
 
-namespace UE::AnimNext::StateTree
+namespace UE::UAF::StateTree
 {
 
 void FAnimNextStateTreeModule::ShutdownModule()
diff --git a/Engine/Plugins/Experimental/AnimNextStateTree/Source/AnimNextStateTree/Private/Conditions/AnimNextStateTreeRigVMConditionBase.cpp b/Engine/Plugins/Experimental/AnimNextStateTree/Source/AnimNextStateTree/Private/Conditions/AnimNextStateTreeRigVMConditionBase.cpp
--- a/Engine/Plugins/Experimental/AnimNextStateTree/Source/AnimNextStateTree/Private/Conditions/AnimNextStateTreeRigVMConditionBase.cpp
+++ b/Engine/Plugins/Experimental/AnimNextStateTree/Source/AnimNextStateTree/Private/Conditions/AnimNextStateTreeRigVMConditionBase.cpp
@@ -63,7 +63,7 @@ bool FAnimNextStateTreeRigVMConditionBase::TestCondition(FStateTreeExecutionCont
 			FAnimNextExecuteContext& AnimNextContext = ExtendedExecuteContext.GetPublicDataSafe<FAnimNextExecuteContext>();
 
 			FAnimNextGraphContextData ContextData(GraphInstance.GetModuleInstance(), &GraphInstance);
-			UE::AnimNext::FScopedExecuteContextData ContextDataScope(AnimNextContext, ContextData);
+			UE::UAF::FScopedExecuteContextData ContextDataScope(AnimNextContext, ContextData);
 
 			RigVMInstance->CallFunctionHandle(InstanceData->FunctionHandle, ExtendedExecuteContext, InstanceData->Parameters);
 
@@ -94,8 +94,8 @@ void FAnimNextStateTreeRigVMConditionBase::PostEditNodeChangeChainProperty(const
 				FRigVMGraphFunctionHeader Result = FRigVMGraphFunctionHeader();
 
 				TMap<FAssetData, FRigVMGraphFunctionHeaderArray> FunctionExports;
-				UE::AnimNext::UncookedOnly::FUtils::GetExportedFunctionsFromAssetRegistry(UE::AnimNext::AnimNextPublicGraphFunctionsExportsRegistryTag, FunctionExports);
-				UE::AnimNext::UncookedOnly::FUtils::GetExportedFunctionsFromAssetRegistry(UE::AnimNext::ControlRigAssetPublicGraphFunctionsExportsRegistryTag, FunctionExports);
+				UE::UAF::UncookedOnly::FUtils::GetExportedFunctionsFromAssetRegistry(UE::UAF::AnimNextPublicGraphFunctionsExportsRegistryTag, FunctionExports);
+				UE::UAF::UncookedOnly::FUtils::GetExportedFunctionsFromAssetRegistry(UE::UAF::ControlRigAssetPublicGraphFunctionsExportsRegistryTag, FunctionExports);
 
 				for (const auto& Export : FunctionExports.Array())
 				{
@@ -144,7 +144,7 @@ void FAnimNextStateTreeRigVMConditionBase::PostEditNodeChangeChainProperty(const
 
 void FAnimNextStateTreeRigVMConditionBase::GetProgrammaticFunctionHeaders(FAnimNextStateTreeProgrammaticFunctionHeaderParams& InProgrammaticFunctionHeaderParams, const UStateTreeState* State, const FStateTreeBindableStructDesc& Desc)
 {
-	using namespace UE::AnimNext::UncookedOnly;
+	using namespace UE::UAF::UncookedOnly;
 
 	FAnimNextGetFunctionHeaderCompileContext& OutCompileContext = InProgrammaticFunctionHeaderParams.OutCompileContext;
 	
diff --git a/Engine/Plugins/Experimental/AnimNextStateTree/Source/AnimNextStateTree/Private/Tasks/AnimNextStateTreeGraphInstanceTask.cpp b/Engine/Plugins/Experimental/AnimNextStateTree/Source/AnimNextStateTree/Private/Tasks/AnimNextStateTreeGraphInstanceTask.cpp
--- a/Engine/Plugins/Experimental/AnimNextStateTree/Source/AnimNextStateTree/Private/Tasks/AnimNextStateTreeGraphInstanceTask.cpp
+++ b/Engine/Plugins/Experimental/AnimNextStateTree/Source/AnimNextStateTree/Private/Tasks/AnimNextStateTreeGraphInstanceTask.cpp
@@ -24,7 +24,7 @@ void FAnimNextGraphInstanceTaskInstanceData::PostSerialize(const FArchive& Ar)
 		// Prior to this we only supported 'imprecise' payloads so try to do our best be existing content, but some data loss is possible
 		if (Asset != nullptr && !Asset->IsA<UAnimNextAnimationGraph>())
 		{
-			Parameters = UE::AnimNext::FAnimGraphFactory::GetDefaultParamsForObject(Asset);
+			Parameters = UE::UAF::FAnimGraphFactory::GetDefaultParamsForObject(Asset);
 
 			// Match any legacy parameter values by name
 			if(Payload_DEPRECATED.GetNumPropertiesInBag() > 0 && Parameters.Builder.Stacks.Num() > 0)
@@ -145,7 +145,7 @@ void FAnimNextStateTreeGraphInstanceTask::PostEditInstanceDataChangeChainPropert
 		else
 		{
 			// TODO: Copy over matching values before overwrite to avoid losing user edits to matching trait data
-			InstanceData.Parameters = UE::AnimNext::FAnimGraphFactory::GetDefaultParamsForObject(InstanceData.Asset);
+			InstanceData.Parameters = UE::UAF::FAnimGraphFactory::GetDefaultParamsForObject(InstanceData.Asset);
 		}
 	}
 }
diff --git a/Engine/Plugins/Experimental/AnimNextStateTree/Source/AnimNextStateTree/Private/Tasks/AnimNextStateTreeRigVMTaskBase.cpp b/Engine/Plugins/Experimental/AnimNextStateTree/Source/AnimNextStateTree/Private/Tasks/AnimNextStateTreeRigVMTaskBase.cpp
--- a/Engine/Plugins/Experimental/AnimNextStateTree/Source/AnimNextStateTree/Private/Tasks/AnimNextStateTreeRigVMTaskBase.cpp
+++ b/Engine/Plugins/Experimental/AnimNextStateTree/Source/AnimNextStateTree/Private/Tasks/AnimNextStateTreeRigVMTaskBase.cpp
@@ -64,7 +64,7 @@ EStateTreeRunStatus FAnimNextStateTreeRigVMTaskBase::EnterState(FStateTreeExecut
 			FAnimNextExecuteContext& AnimNextContext = ExtendedExecuteContext.GetPublicDataSafe<FAnimNextExecuteContext>();
 
 			FAnimNextGraphContextData ContextData(GraphInstance.GetModuleInstance(), &GraphInstance);
-			UE::AnimNext::FScopedExecuteContextData ContextDataScope(AnimNextContext, ContextData);
+			UE::UAF::FScopedExecuteContextData ContextDataScope(AnimNextContext, ContextData);
 
 			// @TODO: Add option to execute per tick?
 			RigVMAsset->CallFunctionHandle(InstanceData->FunctionHandle, ExtendedExecuteContext, InstanceData->Parameters);
@@ -99,8 +99,8 @@ void FAnimNextStateTreeRigVMTaskBase::PostEditNodeChangeChainProperty(const FPro
 				FRigVMGraphFunctionHeader Result = FRigVMGraphFunctionHeader();
 
 				TMap<FAssetData, FRigVMGraphFunctionHeaderArray> FunctionExports;
-				UE::AnimNext::UncookedOnly::FUtils::GetExportedFunctionsFromAssetRegistry(UE::AnimNext::AnimNextPublicGraphFunctionsExportsRegistryTag, FunctionExports);
-				UE::AnimNext::UncookedOnly::FUtils::GetExportedFunctionsFromAssetRegistry(UE::AnimNext::ControlRigAssetPublicGraphFunctionsExportsRegistryTag, FunctionExports);
+				UE::UAF::UncookedOnly::FUtils::GetExportedFunctionsFromAssetRegistry(UE::UAF::AnimNextPublicGraphFunctionsExportsRegistryTag, FunctionExports);
+				UE::UAF::UncookedOnly::FUtils::GetExportedFunctionsFromAssetRegistry(UE::UAF::ControlRigAssetPublicGraphFunctionsExportsRegistryTag, FunctionExports);
 
 				for (const auto& Export : FunctionExports.Array())
 				{
@@ -146,7 +146,7 @@ void FAnimNextStateTreeRigVMTaskBase::PostEditNodeChangeChainProperty(const FPro
 
 void FAnimNextStateTreeRigVMTaskBase::GetProgrammaticFunctionHeaders(FAnimNextStateTreeProgrammaticFunctionHeaderParams& InProgrammaticFunctionHeaderParams, const UStateTreeState* State, const FStateTreeBindableStructDesc& Desc)
 {
-	using namespace UE::AnimNext::UncookedOnly;
+	using namespace UE::UAF::UncookedOnly;
 
 	FAnimNextGetFunctionHeaderCompileContext& OutCompileContext = InProgrammaticFunctionHeaderParams.OutCompileContext;
 	
diff --git a/Engine/Plugins/Experimental/AnimNextStateTree/Source/AnimNextStateTreeEditor/Private/AnimNextStateTreeEditorHost.cpp b/Engine/Plugins/Experimental/AnimNextStateTree/Source/AnimNextStateTreeEditor/Private/AnimNextStateTreeEditorHost.cpp
--- a/Engine/Plugins/Experimental/AnimNextStateTree/Source/AnimNextStateTreeEditor/Private/AnimNextStateTreeEditorHost.cpp
+++ b/Engine/Plugins/Experimental/AnimNextStateTree/Source/AnimNextStateTreeEditor/Private/AnimNextStateTreeEditorHost.cpp
@@ -34,12 +34,12 @@ UStateTree* FAnimNextStateTreeEditorHost::GetStateTree() const
 
 FName FAnimNextStateTreeEditorHost::GetCompilerLogName() const
 {
-	return UE::AnimNext::Editor::LogListingName;
+	return UE::UAF::Editor::LogListingName;
 }
 
 FName FAnimNextStateTreeEditorHost::GetCompilerTabName() const
 {
-	return UE::AnimNext::Editor::CompilerResultsTabName;
+	return UE::UAF::Editor::CompilerResultsTabName;
 }
 
 bool FAnimNextStateTreeEditorHost::ShouldShowCompileButton() const
diff --git a/Engine/Plugins/Experimental/AnimNextStateTree/Source/AnimNextStateTreeEditor/Private/StateTreeAssetCompilationHandler.cpp b/Engine/Plugins/Experimental/AnimNextStateTree/Source/AnimNextStateTreeEditor/Private/StateTreeAssetCompilationHandler.cpp
--- a/Engine/Plugins/Experimental/AnimNextStateTree/Source/AnimNextStateTreeEditor/Private/StateTreeAssetCompilationHandler.cpp
+++ b/Engine/Plugins/Experimental/AnimNextStateTree/Source/AnimNextStateTreeEditor/Private/StateTreeAssetCompilationHandler.cpp
@@ -12,7 +12,7 @@
 #include "StateTreeEditingSubsystem.h"
 #include "UncookedOnlyUtils.h"
 
-namespace UE::AnimNext::StateTree
+namespace UE::UAF::StateTree
 {
 
 FStateTreeAssetCompilationHandler::FStateTreeAssetCompilationHandler(UObject* InAsset)
@@ -51,7 +51,7 @@ void FStateTreeAssetCompilationHandler::Compile(TSharedRef<Workspace::IWorkspace
 		return;
 	}
 
-	UE::AnimNext::UncookedOnly::FScopedCompileJob CompileJob(AnimNextStateTree);
+	UE::UAF::UncookedOnly::FScopedCompileJob CompileJob(AnimNextStateTree);
 
 	// We can't recompile state tree while in PIE
 	if (!GEditor->IsPlaySessionInProgress())
@@ -73,33 +73,33 @@ void FStateTreeAssetCompilationHandler::Compile(TSharedRef<Workspace::IWorkspace
 	}
 
 	// Call super to compile the RigVM Asset 
-	UE::AnimNext::Editor::FAssetCompilationHandler::Compile(InWorkspaceEditor, InAsset);
+	UE::UAF::Editor::FAssetCompilationHandler::Compile(InWorkspaceEditor, InAsset);
 }
 
-UE::AnimNext::Editor::ECompileStatus FStateTreeAssetCompilationHandler::GetCompileStatus(TSharedRef<Workspace::IWorkspaceEditor> InWorkspaceEditor, const UObject* InAsset) const
+UE::UAF::Editor::ECompileStatus FStateTreeAssetCompilationHandler::GetCompileStatus(TSharedRef<Workspace::IWorkspaceEditor> InWorkspaceEditor, const UObject* InAsset) const
 {
 	const UAnimNextStateTree* AnimNextStateTree = Cast<UAnimNextStateTree>(InAsset);
 	if (AnimNextStateTree == nullptr)
 	{
-		return UE::AnimNext::Editor::ECompileStatus::Unknown;
+		return UE::UAF::Editor::ECompileStatus::Unknown;
 	}
 
-	UE::AnimNext::Editor::ECompileStatus RigVMAssetStatus = UE::AnimNext::Editor::FAssetCompilationHandler::GetCompileStatus(InWorkspaceEditor, InAsset);
-	UE::AnimNext::Editor::ECompileStatus StateTreeAssetStatus = UE::AnimNext::Editor::ECompileStatus::Unknown;
+	UE::UAF::Editor::ECompileStatus RigVMAssetStatus = UE::UAF::Editor::FAssetCompilationHandler::GetCompileStatus(InWorkspaceEditor, InAsset);
+	UE::UAF::Editor::ECompileStatus StateTreeAssetStatus = UE::UAF::Editor::ECompileStatus::Unknown;
 	if (AnimNextStateTree->StateTree)
 	{
 		const bool bCompiledDataResetDuringLoad = AnimNextStateTree->StateTree->LastCompiledEditorDataHash == EditorDataHash && !AnimNextStateTree->StateTree->IsReadyToRun();
 		if (!bLastCompileSucceeded || bCompiledDataResetDuringLoad)
 		{
-			StateTreeAssetStatus = UE::AnimNext::Editor::ECompileStatus::Error;
+			StateTreeAssetStatus = UE::UAF::Editor::ECompileStatus::Error;
 		}
 		else if (AnimNextStateTree->StateTree->LastCompiledEditorDataHash != EditorDataHash)
 		{
-			StateTreeAssetStatus = UE::AnimNext::Editor::ECompileStatus::Dirty;
+			StateTreeAssetStatus = UE::UAF::Editor::ECompileStatus::Dirty;
 		}
 		else
 		{
-			StateTreeAssetStatus = UE::AnimNext::Editor::ECompileStatus::UpToDate;
+			StateTreeAssetStatus = UE::UAF::Editor::ECompileStatus::UpToDate;
 		}
 	}
 
diff --git a/Engine/Plugins/Experimental/AnimNextStateTree/Source/AnimNextStateTreeEditor/Private/StateTreeAssetCompilationHandler.h b/Engine/Plugins/Experimental/AnimNextStateTree/Source/AnimNextStateTreeEditor/Private/StateTreeAssetCompilationHandler.h
--- a/Engine/Plugins/Experimental/AnimNextStateTree/Source/AnimNextStateTreeEditor/Private/StateTreeAssetCompilationHandler.h
+++ b/Engine/Plugins/Experimental/AnimNextStateTree/Source/AnimNextStateTreeEditor/Private/StateTreeAssetCompilationHandler.h
@@ -9,10 +9,10 @@
 class UStateTreeState;
 class UStateTree;
 
-namespace UE::AnimNext::StateTree
+namespace UE::UAF::StateTree
 {
 
-class FStateTreeAssetCompilationHandler : public UE::AnimNext::Editor::FAssetCompilationHandler
+class FStateTreeAssetCompilationHandler : public UE::UAF::Editor::FAssetCompilationHandler
 {
 public:
 	explicit FStateTreeAssetCompilationHandler(UObject* InObject);
@@ -23,7 +23,7 @@ public:
 private:
 	// IAssetCompilationHandler interface
 	virtual void Compile(TSharedRef<Workspace::IWorkspaceEditor> InWorkspaceEditor, UObject* InAsset) override;
-	virtual UE::AnimNext::Editor::ECompileStatus GetCompileStatus(TSharedRef<Workspace::IWorkspaceEditor> InWorkspaceEditor, const UObject* InAsset) const override;
+	virtual UE::UAF::Editor::ECompileStatus GetCompileStatus(TSharedRef<Workspace::IWorkspaceEditor> InWorkspaceEditor, const UObject* InAsset) const override;
 
 	void UpdateCachedInfo();
 
diff --git a/Engine/Plugins/Experimental/AnimNextStateTree/Source/AnimNextStateTreeUncookedOnly/Internal/AnimNextStateTreeEditorData.cpp b/Engine/Plugins/Experimental/AnimNextStateTree/Source/AnimNextStateTreeUncookedOnly/Internal/AnimNextStateTreeEditorData.cpp
--- a/Engine/Plugins/Experimental/AnimNextStateTree/Source/AnimNextStateTreeUncookedOnly/Internal/AnimNextStateTreeEditorData.cpp
+++ b/Engine/Plugins/Experimental/AnimNextStateTree/Source/AnimNextStateTreeUncookedOnly/Internal/AnimNextStateTreeEditorData.cpp
@@ -11,7 +11,7 @@ void UAnimNextStateTreeTreeEditorData::PostLoad()
 	Super::PostLoad();
 	if(const UAnimNextStateTree* OuterAnimNextStateTree = GetTypedOuter<UAnimNextStateTree>())
 	{
-		UAnimNextStateTree_EditorData* AnimNextEditorData = UE::AnimNext::UncookedOnly::FUtils::GetEditorData<UAnimNextStateTree_EditorData>(OuterAnimNextStateTree);	
+		UAnimNextStateTree_EditorData* AnimNextEditorData = UE::UAF::UncookedOnly::FUtils::GetEditorData<UAnimNextStateTree_EditorData>(OuterAnimNextStateTree);	
 		AnimNextEditorData->ModifiedDelegate.AddUObject(this, &UAnimNextStateTreeTreeEditorData::HandleStateTreeAssetChanges);
 	}
 }
@@ -21,7 +21,7 @@ void UAnimNextStateTreeTreeEditorData::CreateRootProperties(TArrayView<UE::Prope
 	UStateTree* OuterStateTree = CastChecked<UStateTree>(GetOuter());
 	UAnimNextStateTree* OuterAnimNextStateTree = CastChecked<UAnimNextStateTree>(OuterStateTree->GetOuter());
 
-	UAnimNextStateTree_EditorData* AnimNextEditorData = UE::AnimNext::UncookedOnly::FUtils::GetEditorData<UAnimNextStateTree_EditorData>(OuterAnimNextStateTree);
+	UAnimNextStateTree_EditorData* AnimNextEditorData = UE::UAF::UncookedOnly::FUtils::GetEditorData<UAnimNextStateTree_EditorData>(OuterAnimNextStateTree);
 
 	// Generate unique names for the incoming property descs to avoid changing the existing properties in the bag
 	for (UE::PropertyBinding::FPropertyCreationDescriptor& CreationDesc : InOutCreationDescs)
@@ -71,6 +71,6 @@ const FInstancedPropertyBag& UAnimNextStateTreeTreeEditorData::GetRootParameters
 {
 	UStateTree* OuterStateTree = CastChecked<UStateTree>(GetOuter());
 	UAnimNextStateTree* OuterAnimNextStateTree = CastChecked<UAnimNextStateTree>(OuterStateTree->GetOuter());
-	UAnimNextStateTree_EditorData* AnimNextEditorData = UE::AnimNext::UncookedOnly::FUtils::GetEditorData<UAnimNextStateTree_EditorData>(OuterAnimNextStateTree);
+	UAnimNextStateTree_EditorData* AnimNextEditorData = UE::UAF::UncookedOnly::FUtils::GetEditorData<UAnimNextStateTree_EditorData>(OuterAnimNextStateTree);
 	return AnimNextEditorData->CombinedVariablesPropertyBag;
 }
diff --git a/Engine/Plugins/Experimental/AnimNextStateTree/Source/AnimNextStateTreeUncookedOnly/Internal/AnimNextStateTree_EditorData.cpp b/Engine/Plugins/Experimental/AnimNextStateTree/Source/AnimNextStateTreeUncookedOnly/Internal/AnimNextStateTree_EditorData.cpp
--- a/Engine/Plugins/Experimental/AnimNextStateTree/Source/AnimNextStateTreeUncookedOnly/Internal/AnimNextStateTree_EditorData.cpp
+++ b/Engine/Plugins/Experimental/AnimNextStateTree/Source/AnimNextStateTreeUncookedOnly/Internal/AnimNextStateTree_EditorData.cpp
@@ -46,7 +46,7 @@ void UAnimNextStateTree_EditorData::RecompileVM()
 		return;
 	}
 
-	UAnimNextStateTree* AnimationStateTree = UE::AnimNext::UncookedOnly::FUtils::GetAsset<UAnimNextStateTree>(this);
+	UAnimNextStateTree* AnimationStateTree = UE::UAF::UncookedOnly::FUtils::GetAsset<UAnimNextStateTree>(this);
 
 	UStateTree* InnerStateTree = AnimationStateTree->StateTree;
 
@@ -72,7 +72,7 @@ void UAnimNextStateTree_EditorData::OnPreCompileGetProgrammaticFunctionHeaders(c
 {
 	Super::OnPreCompileGetProgrammaticFunctionHeaders(InSettings, OutCompileContext);
 
-	if (UAnimNextStateTree* AnimStateTree = UE::AnimNext::UncookedOnly::FUtils::GetAsset<UAnimNextStateTree>(this))
+	if (UAnimNextStateTree* AnimStateTree = UE::UAF::UncookedOnly::FUtils::GetAsset<UAnimNextStateTree>(this))
 	{
 		// Give child nodes a chance to add compile-time only variables
 		FAnimNextStateTreeProgrammaticFunctionHeaderParams ProgrammaticFunctionHeaderParams(this, InSettings, RigVMClient, OutCompileContext);
@@ -131,13 +131,13 @@ void UAnimNextStateTree_EditorData::OnPreCompileGetProgrammaticGraphs(const FRig
 {
 	Super::OnPreCompileGetProgrammaticGraphs(InSettings, OutCompileContext);
 
-	if (UAnimNextStateTree* AnimStateTree = UE::AnimNext::UncookedOnly::FUtils::GetAsset<UAnimNextStateTree>(this))
+	if (UAnimNextStateTree* AnimStateTree = UE::UAF::UncookedOnly::FUtils::GetAsset<UAnimNextStateTree>(this))
 	{
 		URigVMGraph* Graph = NewObject<URigVMGraph>(this, NAME_None, RF_Transient);
 		Graph->SetSchemaClass(UAnimNextAnimationGraphSchema::StaticClass());
 
 		UAnimNextController* Controller = CastChecked<UAnimNextController>(RigVMClient.GetOrCreateController(Graph));
-		UE::AnimNext::UncookedOnly::FAnimGraphUtils::SetupAnimGraph(FRigUnit_AnimNextGraphRoot::DefaultEntryPoint, Controller, false);
+		UE::UAF::UncookedOnly::FAnimGraphUtils::SetupAnimGraph(FRigUnit_AnimNextGraphRoot::DefaultEntryPoint, Controller, false);
 
 		if (Controller->GetGraph()->GetNodes().Num() != 1)
 		{
@@ -162,7 +162,7 @@ void UAnimNextStateTree_EditorData::OnPreCompileGetProgrammaticGraphs(const FRig
 			return;
 		}
 
-		const FName BlendStackTraitName = Controller->AddTraitByName(TraitStackNode->GetFName(), *UE::AnimNext::FTraitRegistry::Get().Find(UE::AnimNext::FBlendStackCoreTrait::TraitUID)->GetTraitName(), INDEX_NONE, TEXT(""), false);
+		const FName BlendStackTraitName = Controller->AddTraitByName(TraitStackNode->GetFName(), *UE::UAF::FTraitRegistry::Get().Find(UE::UAF::FBlendStackCoreTrait::TraitUID)->GetTraitName(), INDEX_NONE, TEXT(""), false);
 		
 		if (BlendStackTraitName == NAME_None)
 		{
@@ -170,7 +170,7 @@ void UAnimNextStateTree_EditorData::OnPreCompileGetProgrammaticGraphs(const FRig
 			return;
 		}
 		
-		const FName StateTreeTraitName = Controller->AddTraitByName(TraitStackNode->GetFName(), *UE::AnimNext::FTraitRegistry::Get().Find(UE::AnimNext::FStateTreeTrait::TraitUID)->GetTraitName(), INDEX_NONE, TEXT(""), false);
+		const FName StateTreeTraitName = Controller->AddTraitByName(TraitStackNode->GetFName(), *UE::UAF::FTraitRegistry::Get().Find(UE::UAF::FStateTreeTrait::TraitUID)->GetTraitName(), INDEX_NONE, TEXT(""), false);
 
 		if (StateTreeTraitName == NAME_None)
 		{
@@ -178,7 +178,7 @@ void UAnimNextStateTree_EditorData::OnPreCompileGetProgrammaticGraphs(const FRig
 			return;
 		}
 
-		const FName SmootherTraitName = Controller->AddTraitByName(TraitStackNode->GetFName(), *UE::AnimNext::FTraitRegistry::Get().Find(UE::AnimNext::FBlendSmootherCoreTrait::TraitUID)->GetTraitName(), INDEX_NONE, TEXT(""), false);
+		const FName SmootherTraitName = Controller->AddTraitByName(TraitStackNode->GetFName(), *UE::UAF::FTraitRegistry::Get().Find(UE::UAF::FBlendSmootherCoreTrait::TraitUID)->GetTraitName(), INDEX_NONE, TEXT(""), false);
 
 		if (SmootherTraitName == NAME_None)
 		{
diff --git a/Engine/Plugins/Experimental/AnimNextStateTree/Source/AnimNextStateTreeUncookedOnly/Private/AnimNextStateTreeFunctionLibraryHelper.cpp b/Engine/Plugins/Experimental/AnimNextStateTree/Source/AnimNextStateTreeUncookedOnly/Private/AnimNextStateTreeFunctionLibraryHelper.cpp
--- a/Engine/Plugins/Experimental/AnimNextStateTree/Source/AnimNextStateTreeUncookedOnly/Private/AnimNextStateTreeFunctionLibraryHelper.cpp
+++ b/Engine/Plugins/Experimental/AnimNextStateTree/Source/AnimNextStateTreeUncookedOnly/Private/AnimNextStateTreeFunctionLibraryHelper.cpp
@@ -10,8 +10,8 @@ const TArray<FName> UAnimNextStateTreeFunctionLibraryHelper::GetExposedAnimNextF
 	TArray<FName> Result;
 
 	TMap<FAssetData, FRigVMGraphFunctionHeaderArray> FunctionExports;
-	UE::AnimNext::UncookedOnly::FUtils::GetExportedFunctionsFromAssetRegistry(UE::AnimNext::AnimNextPublicGraphFunctionsExportsRegistryTag, FunctionExports);
-	UE::AnimNext::UncookedOnly::FUtils::GetExportedFunctionsFromAssetRegistry(UE::AnimNext::ControlRigAssetPublicGraphFunctionsExportsRegistryTag, FunctionExports);
+	UE::UAF::UncookedOnly::FUtils::GetExportedFunctionsFromAssetRegistry(UE::UAF::AnimNextPublicGraphFunctionsExportsRegistryTag, FunctionExports);
+	UE::UAF::UncookedOnly::FUtils::GetExportedFunctionsFromAssetRegistry(UE::UAF::ControlRigAssetPublicGraphFunctionsExportsRegistryTag, FunctionExports);
 	
 	for (const TPair<FAssetData, FRigVMGraphFunctionHeaderArray>& Export : FunctionExports.Array())
 	{
diff --git a/Engine/Plugins/Experimental/AnimNextStateTree/Source/AnimNextStateTreeUncookedOnly/Private/AnimNextStateTreeUncookedModule.cpp b/Engine/Plugins/Experimental/AnimNextStateTree/Source/AnimNextStateTreeUncookedOnly/Private/AnimNextStateTreeUncookedModule.cpp
--- a/Engine/Plugins/Experimental/AnimNextStateTree/Source/AnimNextStateTreeUncookedOnly/Private/AnimNextStateTreeUncookedModule.cpp
+++ b/Engine/Plugins/Experimental/AnimNextStateTree/Source/AnimNextStateTreeUncookedOnly/Private/AnimNextStateTreeUncookedModule.cpp
@@ -2,4 +2,4 @@
 
 #include "AnimNextStateTreeUncookedModule.h"
 
-IMPLEMENT_MODULE(UE::AnimNext::StateTree::FAnimNextStateTreeUncookedModule, AnimNextStateTreeUncookedOnly)
\ No newline at end of file
+IMPLEMENT_MODULE(UE::UAF::StateTree::FAnimNextStateTreeUncookedModule, AnimNextStateTreeUncookedOnly)
\ No newline at end of file
diff --git a/Engine/Plugins/Experimental/AnimNextWarping/Source/AnimNextWarping/Private/Module.cpp b/Engine/Plugins/Experimental/AnimNextWarping/Source/AnimNextWarping/Private/Module.cpp
--- a/Engine/Plugins/Experimental/AnimNextWarping/Source/AnimNextWarping/Private/Module.cpp
+++ b/Engine/Plugins/Experimental/AnimNextWarping/Source/AnimNextWarping/Private/Module.cpp
@@ -13,7 +13,7 @@
 
 DEFINE_LOG_CATEGORY(LogAnimNextWarping);
 
-namespace UE::AnimNext::Warping
+namespace UE::UAF::Warping
 {
 
 class FModule : public IModuleInterface
@@ -46,4 +46,4 @@ public:
 
 }
 
-IMPLEMENT_MODULE(UE::AnimNext::Warping::FModule, AnimNextWarping)
+IMPLEMENT_MODULE(UE::UAF::Warping::FModule, AnimNextWarping)
diff --git a/Engine/Plugins/Experimental/AnimNextWarping/Source/AnimNextWarping/Private/SteeringTrait.cpp b/Engine/Plugins/Experimental/AnimNextWarping/Source/AnimNextWarping/Private/SteeringTrait.cpp
--- a/Engine/Plugins/Experimental/AnimNextWarping/Source/AnimNextWarping/Private/SteeringTrait.cpp
+++ b/Engine/Plugins/Experimental/AnimNextWarping/Source/AnimNextWarping/Private/SteeringTrait.cpp
@@ -24,7 +24,7 @@ static FAutoConsoleVariableRef CVarAnimNextSteeringTraitEnabled(
 ////////////////////////////////////////////////////////////////////////////////////////////////////////////////////////
 // FSteeringTrait
 
-namespace UE::AnimNext
+namespace UE::UAF
 {
 	AUTO_REGISTER_ANIM_TRAIT(FSteeringTrait)
 
@@ -90,13 +90,13 @@ namespace UE::AnimNext
 
 		Context.AppendTask(FAnimNextSteeringTask::Make(InstanceData, SharedData));
 	}
-} // namespace UE::AnimNext
+} // namespace UE::UAF
 
 
 ////////////////////////////////////////////////////////////////////////////////////////////////////////////////////////
 // FAnimNextSteeringTask
 
-FAnimNextSteeringTask FAnimNextSteeringTask::Make(UE::AnimNext::FSteeringTrait::FInstanceData* InstanceData, const UE::AnimNext::FSteeringTrait::FSharedData* SharedData)
+FAnimNextSteeringTask FAnimNextSteeringTask::Make(UE::UAF::FSteeringTrait::FInstanceData* InstanceData, const UE::UAF::FSteeringTrait::FSharedData* SharedData)
 {
 	FAnimNextSteeringTask Task;
 	Task.InstanceData = InstanceData;
@@ -104,9 +104,9 @@ FAnimNextSteeringTask FAnimNextSteeringTask::Make(UE::AnimNext::FSteeringTrait::
 	return Task;
 }
 
-void FAnimNextSteeringTask::Execute(UE::AnimNext::FEvaluationVM& VM) const
+void FAnimNextSteeringTask::Execute(UE::UAF::FEvaluationVM& VM) const
 {
-	using namespace UE::AnimNext;
+	using namespace UE::UAF;
 
 	// Execute Steering
 	if (InstanceData->DeltaTime > 0.f)
diff --git a/Engine/Plugins/Experimental/AnimNextWarping/Source/AnimNextWarping/Private/SteeringTrait.h b/Engine/Plugins/Experimental/AnimNextWarping/Source/AnimNextWarping/Private/SteeringTrait.h
--- a/Engine/Plugins/Experimental/AnimNextWarping/Source/AnimNextWarping/Private/SteeringTrait.h
+++ b/Engine/Plugins/Experimental/AnimNextWarping/Source/AnimNextWarping/Private/SteeringTrait.h
@@ -86,7 +86,7 @@ struct FSteeringTraitSharedData : public FAnimNextTraitSharedData
  * 
  * Ex: If your anim only rotates say 45 deg, but you need to warp it to 60 deg to match gameplay input
  */
-namespace UE::AnimNext
+namespace UE::UAF
 {
 
 struct FSteeringTrait : FAdditiveTrait, IUpdate, IEvaluate
@@ -131,7 +131,7 @@ struct FSteeringTrait : FAdditiveTrait, IUpdate, IEvaluate
 	virtual void PostEvaluate(FEvaluateTraversalContext& Context, const TTraitBinding<IEvaluate>& Binding) const override;
 };
 
-} // namespace UE::AnimNext
+} // namespace UE::UAF
 
 /** Task to run Steering on VM */
 USTRUCT()
@@ -141,11 +141,11 @@ struct FAnimNextSteeringTask : public FAnimNextEvaluationTask
 
 	DECLARE_ANIM_EVALUATION_TASK(FAnimNextSteeringTask)
 
-	static FAnimNextSteeringTask Make(UE::AnimNext::FSteeringTrait::FInstanceData* InstanceData, const UE::AnimNext::FSteeringTrait::FSharedData* SharedData);
+	static FAnimNextSteeringTask Make(UE::UAF::FSteeringTrait::FInstanceData* InstanceData, const UE::UAF::FSteeringTrait::FSharedData* SharedData);
 
 	// Task entry point
-	virtual void Execute(UE::AnimNext::FEvaluationVM& VM) const override;
+	virtual void Execute(UE::UAF::FEvaluationVM& VM) const override;
 
-	UE::AnimNext::FSteeringTrait::FInstanceData* InstanceData = nullptr;
-	const UE::AnimNext::FSteeringTrait::FSharedData* SharedData = nullptr;
+	UE::UAF::FSteeringTrait::FInstanceData* InstanceData = nullptr;
+	const UE::UAF::FSteeringTrait::FSharedData* SharedData = nullptr;
 };
\ No newline at end of file
diff --git a/Engine/Plugins/Experimental/AnimNextWarping/Source/AnimNextWarping/Private/StrafeWarpingTrait.cpp b/Engine/Plugins/Experimental/AnimNextWarping/Source/AnimNextWarping/Private/StrafeWarpingTrait.cpp
--- a/Engine/Plugins/Experimental/AnimNextWarping/Source/AnimNextWarping/Private/StrafeWarpingTrait.cpp
+++ b/Engine/Plugins/Experimental/AnimNextWarping/Source/AnimNextWarping/Private/StrafeWarpingTrait.cpp
@@ -27,7 +27,7 @@ constexpr bool bAnimNextStrafeWarpingTraitEnabled = true;
 ////////////////////////////////////////////////////////////////////////////////////////////////////////////////////////
 // FStrafeWarpingTrait
 
-namespace UE::AnimNext
+namespace UE::UAF
 {
 	static FVector GetAxisVector(const EAxis::Type& InAxis)
 	{
@@ -165,13 +165,13 @@ namespace UE::AnimNext
 			}
 		}
 	}
-} // namespace UE::AnimNext
+} // namespace UE::UAF
 
 
 ////////////////////////////////////////////////////////////////////////////////////////////////////////////////////////
 // FAnimNextSteeringTask
 
-FAnimNextStrafeWarpingTask FAnimNextStrafeWarpingTask::Make(UE::AnimNext::FStrafeWarpingTrait::FInstanceData* InstanceData, const UE::AnimNext::FStrafeWarpingTrait::FSharedData* SharedData)
+FAnimNextStrafeWarpingTask FAnimNextStrafeWarpingTask::Make(UE::UAF::FStrafeWarpingTrait::FInstanceData* InstanceData, const UE::UAF::FStrafeWarpingTrait::FSharedData* SharedData)
 {
 	FAnimNextStrafeWarpingTask Task;
 	Task.InstanceData = InstanceData;
@@ -181,9 +181,9 @@ FAnimNextStrafeWarpingTask FAnimNextStrafeWarpingTask::Make(UE::AnimNext::FStraf
 
 
 
-void FAnimNextStrafeWarpingTask::Execute(UE::AnimNext::FEvaluationVM& VM) const
+void FAnimNextStrafeWarpingTask::Execute(UE::UAF::FEvaluationVM& VM) const
 {
-	using namespace UE::AnimNext;
+	using namespace UE::UAF;
 	QUICK_SCOPE_CYCLE_COUNTER(FAnimNextStrafeWarpingTask_Execute);
 	
 	if (InstanceData->DeltaTime == 0.f)
@@ -245,7 +245,7 @@ void FAnimNextStrafeWarpingTask::Execute(UE::AnimNext::FEvaluationVM& VM) const
 	DoWarpPose(*Keyframe, RotationAxisVector);
 }
 
-void FAnimNextStrafeWarpingTask::DoWarpRootMotion(TUniquePtr<UE::AnimNext::FKeyframeState>& Keyframe, const FVector& RotationAxisVector, const FVector& TargetMoveDir, float& OutTargetOrientationAngleRad) const
+void FAnimNextStrafeWarpingTask::DoWarpRootMotion(TUniquePtr<UE::UAF::FKeyframeState>& Keyframe, const FVector& RotationAxisVector, const FVector& TargetMoveDir, float& OutTargetOrientationAngleRad) const
 {
 	FTransform ThisFrameRootMotionTransform = FTransform::Identity;
 	const UE::Anim::IAnimRootMotionProvider* RootMotionProvider = UE::Anim::IAnimRootMotionProvider::Get();
@@ -272,7 +272,7 @@ void FAnimNextStrafeWarpingTask::DoWarpRootMotion(TUniquePtr<UE::AnimNext::FKeyf
 			const FVector PreviousRootMotionDeltaDirection = InstanceData->RootMotionDeltaDirection;
 			// Hold previous direction if we can't calculate it from current move delta, because the root is no longer moving
 			InstanceData->RootMotionDeltaDirection = RootMotionDeltaTranslation.GetSafeNormal(UE_SMALL_NUMBER, PreviousRootMotionDeltaDirection);
-			OutTargetOrientationAngleRad = UE::AnimNext::SignedAngleRadBetweenNormals(
+			OutTargetOrientationAngleRad = UE::UAF::SignedAngleRadBetweenNormals(
 				InstanceData->RootMotionDeltaDirection, TargetMoveDir, RotationAxisVector);
 
 			// Motion Matching may return an animation that deviates a lot from the movement direction (e.g movement direction going bwd and motion matching could return the fwd animation for a few frames)
@@ -318,7 +318,7 @@ void FAnimNextStrafeWarpingTask::DoWarpRootMotion(TUniquePtr<UE::AnimNext::FKeyf
 				// This prevents our interpolation from fighting the natural root motion that's flowing through the graph.
 				// To correctly measure the amount to counter, we need to unrotate our previous delta direction by our previous rotation
 				// As the previous direction delta is relative to the previous rotation delta
-				RootMotionDeltaAngleRad = UE::AnimNext::SignedAngleRadBetweenNormals(
+				RootMotionDeltaAngleRad = UE::UAF::SignedAngleRadBetweenNormals(
 					InstanceData->RootMotionDeltaDirection, PreviousRootMotionDeltaRotation.UnrotateVector(PreviousRootMotionDeltaDirection),
 					RotationAxisVector);
 
@@ -411,9 +411,9 @@ void FAnimNextStrafeWarpingTask::DoWarpRootMotion(TUniquePtr<UE::AnimNext::FKeyf
 #endif // ENABLE_VISUAL_LOG
 }
 
-void FAnimNextStrafeWarpingTask::DoWarpPose(TUniquePtr<UE::AnimNext::FKeyframeState>& Keyframe, const FVector& RotationAxisVector) const
+void FAnimNextStrafeWarpingTask::DoWarpPose(TUniquePtr<UE::UAF::FKeyframeState>& Keyframe, const FVector& RotationAxisVector) const
 {
-	using namespace UE::AnimNext;
+	using namespace UE::UAF;
 	
 	const FLODPoseStack& Pose = Keyframe->Pose;
 
diff --git a/Engine/Plugins/Experimental/AnimNextWarping/Source/AnimNextWarping/Private/StrafeWarpingTrait.h b/Engine/Plugins/Experimental/AnimNextWarping/Source/AnimNextWarping/Private/StrafeWarpingTrait.h
--- a/Engine/Plugins/Experimental/AnimNextWarping/Source/AnimNextWarping/Private/StrafeWarpingTrait.h
+++ b/Engine/Plugins/Experimental/AnimNextWarping/Source/AnimNextWarping/Private/StrafeWarpingTrait.h
@@ -125,7 +125,7 @@ struct FStrafeWarpingTraitSharedData : public FAnimNextTraitSharedData
 /**
  * This trait modifies an animation pose to orientate towards a desired move direction 
  */
-namespace UE::AnimNext
+namespace UE::UAF
 {
 
 struct FStrafeWarpingTrait : FAdditiveTrait, IUpdate, IEvaluate
@@ -204,7 +204,7 @@ struct FStrafeWarpingTrait : FAdditiveTrait, IUpdate, IEvaluate
 	static void InitializeSpineData(TArrayView<FSpineBoneData> OutSpineBoneData, const TArray<FName>& SpineBoneNames, const FLODPoseStack& Pose);
 };
 
-} // namespace UE::AnimNext
+} // namespace UE::UAF
 
 /** Task to run Strafe Warping on VM */
 USTRUCT()
@@ -214,15 +214,15 @@ struct FAnimNextStrafeWarpingTask : public FAnimNextEvaluationTask
 
 	DECLARE_ANIM_EVALUATION_TASK(FAnimNextStrafeWarpingTask)
 
-	static FAnimNextStrafeWarpingTask Make(UE::AnimNext::FStrafeWarpingTrait::FInstanceData* InstanceData, const UE::AnimNext::FStrafeWarpingTrait::FSharedData* SharedData);
+	static FAnimNextStrafeWarpingTask Make(UE::UAF::FStrafeWarpingTrait::FInstanceData* InstanceData, const UE::UAF::FStrafeWarpingTrait::FSharedData* SharedData);
 
 	// Task entry point
-	virtual void Execute(UE::AnimNext::FEvaluationVM& VM) const override;
+	virtual void Execute(UE::UAF::FEvaluationVM& VM) const override;
 
 private:
-	void DoWarpRootMotion(TUniquePtr<UE::AnimNext::FKeyframeState>& KeyframeState, const FVector& RotationAxisVector, const FVector& TargetMoveDir, float& OutTargetOrientationAngleRad) const;
-	void DoWarpPose(TUniquePtr<UE::AnimNext::FKeyframeState>& KeyframeState, const FVector& RotationAxisVector) const;
+	void DoWarpRootMotion(TUniquePtr<UE::UAF::FKeyframeState>& KeyframeState, const FVector& RotationAxisVector, const FVector& TargetMoveDir, float& OutTargetOrientationAngleRad) const;
+	void DoWarpPose(TUniquePtr<UE::UAF::FKeyframeState>& KeyframeState, const FVector& RotationAxisVector) const;
 
-	UE::AnimNext::FStrafeWarpingTrait::FInstanceData* InstanceData = nullptr;
-	const UE::AnimNext::FStrafeWarpingTrait::FSharedData* SharedData = nullptr;
+	UE::UAF::FStrafeWarpingTrait::FInstanceData* InstanceData = nullptr;
+	const UE::UAF::FStrafeWarpingTrait::FSharedData* SharedData = nullptr;
 };
\ No newline at end of file
diff --git a/Engine/Plugins/Experimental/AnimNextWarping/Source/AnimNextWarping/Private/WarpTestTrait.cpp b/Engine/Plugins/Experimental/AnimNextWarping/Source/AnimNextWarping/Private/WarpTestTrait.cpp
--- a/Engine/Plugins/Experimental/AnimNextWarping/Source/AnimNextWarping/Private/WarpTestTrait.cpp
+++ b/Engine/Plugins/Experimental/AnimNextWarping/Source/AnimNextWarping/Private/WarpTestTrait.cpp
@@ -9,7 +9,7 @@
 #include "Module/AnimNextModuleInstance.h"
 #include "VisualLogger/VisualLogger.h"
 
-namespace UE::AnimNext
+namespace UE::UAF
 {
 	AUTO_REGISTER_ANIM_TRAIT(FWarpTestTrait)
 
@@ -89,11 +89,11 @@ namespace UE::AnimNext
 			}
 		}		
 	}
-} // namespace UE::AnimNext
+} // namespace UE::UAF
 
-void FAnimNextWarpTestTask::Execute(UE::AnimNext::FEvaluationVM& VM) const
+void FAnimNextWarpTestTask::Execute(UE::UAF::FEvaluationVM& VM) const
 {
-	using namespace UE::AnimNext;
+	using namespace UE::UAF;
 
 	const UE::Anim::IAnimRootMotionProvider* RootMotionProvider = UE::Anim::IAnimRootMotionProvider::Get();
 	if (!RootMotionProvider)
diff --git a/Engine/Plugins/Experimental/AnimNextWarping/Source/AnimNextWarping/Private/WarpTestTrait.h b/Engine/Plugins/Experimental/AnimNextWarping/Source/AnimNextWarping/Private/WarpTestTrait.h
--- a/Engine/Plugins/Experimental/AnimNextWarping/Source/AnimNextWarping/Private/WarpTestTrait.h
+++ b/Engine/Plugins/Experimental/AnimNextWarping/Source/AnimNextWarping/Private/WarpTestTrait.h
@@ -30,7 +30,7 @@ struct FWarpTestTraitSharedData : public FAnimNextTraitSharedData
 	#undef TRAIT_LATENT_PROPERTIES_ENUMERATOR
 };
 
-namespace UE::AnimNext
+namespace UE::UAF
 {
 
 struct FWarpTestTrait : FAdditiveTrait, IUpdate, IEvaluate
@@ -52,7 +52,7 @@ struct FWarpTestTrait : FAdditiveTrait, IUpdate, IEvaluate
 	virtual void PostEvaluate(FEvaluateTraversalContext& Context, const TTraitBinding<IEvaluate>& Binding) const override;
 };
 
-} // namespace UE::AnimNext
+} // namespace UE::UAF
 
 USTRUCT(Experimental)
 struct FAnimNextWarpTestTask : public FAnimNextEvaluationTask
@@ -61,7 +61,7 @@ struct FAnimNextWarpTestTask : public FAnimNextEvaluationTask
 
 	DECLARE_ANIM_EVALUATION_TASK(FAnimNextWarpTestTask)
 
-	virtual void Execute(UE::AnimNext::FEvaluationVM& VM) const override;
+	virtual void Execute(UE::UAF::FEvaluationVM& VM) const override;
 
 	FTransform ComponentTransform = FTransform::Identity;
 	FTransform WarpToTransform = FTransform::Identity;
diff --git a/Engine/Plugins/Experimental/EvaluationNotifies/Source/Runtime/Private/AnimNext/AnimNextAlignment.cpp b/Engine/Plugins/Experimental/EvaluationNotifies/Source/Runtime/Private/AnimNext/AnimNextAlignment.cpp
--- a/Engine/Plugins/Experimental/EvaluationNotifies/Source/Runtime/Private/AnimNext/AnimNextAlignment.cpp
+++ b/Engine/Plugins/Experimental/EvaluationNotifies/Source/Runtime/Private/AnimNext/AnimNextAlignment.cpp
@@ -20,7 +20,7 @@ void FEvaluationNotify_AlignmentInstance::Start()
 	bFirstFrame = true;
 }
 
-namespace AnimNextAlignment
+namespace UE::UAF::Alignment
 {
 	constexpr float FrameTime = 1.0f/30.0f;
 	
@@ -63,13 +63,13 @@ float FEvaluationNotify_AlignmentInstance::GetWeight(float Time, const FAlignmen
 	{
 		const float Duration = (RoundedEndTime - StartTime);
 		const float StartCurveSampleTime = FMath::Max(WarpCurve.StartRatio * Duration, ActualStartTime) - ActualStartTime;
-		return AnimNextAlignment::SampleCurve(Time - ActualStartTime, StartCurveSampleTime, WarpCurve.EndRatio * Duration - ActualStartTime, AnimTrajectoryData.TranslationCurve);
+		return UE::UAF::Alignment::SampleCurve(Time - ActualStartTime, StartCurveSampleTime, WarpCurve.EndRatio * Duration - ActualStartTime, AnimTrajectoryData.TranslationCurve);
 	}
 	else if (WarpCurve.CurveType == EAlignmentWeightCurveType::FromRootMotionRotation && !AnimTrajectoryData.RotationCurve.IsEmpty())
 	{
 		const float Duration = (RoundedEndTime - StartTime);
 		const float StartCurveSampleTime = FMath::Max(WarpCurve.StartRatio * Duration, ActualStartTime) - ActualStartTime;
-		return AnimNextAlignment::SampleCurve(Time - ActualStartTime, StartCurveSampleTime, WarpCurve.EndRatio * Duration - ActualStartTime, AnimTrajectoryData.RotationCurve);
+		return UE::UAF::Alignment::SampleCurve(Time - ActualStartTime, StartCurveSampleTime, WarpCurve.EndRatio * Duration - ActualStartTime, AnimTrajectoryData.RotationCurve);
 	}
 	else
 	{
@@ -110,7 +110,7 @@ float FEvaluationNotify_AlignmentInstance::GetWeight(float Time, const FAlignmen
 	return Weight;
 }
 
-void FEvaluationNotify_AlignToGroundInstance::End(UE::AnimNext::FEvaluationNotifiesTrait::FInstanceData& TraitInstanceData)
+void FEvaluationNotify_AlignToGroundInstance::End(UE::UAF::FEvaluationNotifiesTrait::FInstanceData& TraitInstanceData)
 {
 	FEvaluationNotify_AlignmentInstance::End(TraitInstanceData);
 	
@@ -118,7 +118,7 @@ void FEvaluationNotify_AlignToGroundInstance::End(UE::AnimNext::FEvaluationNotif
 	TraitInstanceData.Instance->SetVariable(AlignToGroundNotify->PlaybackRateOutputVariableReference, 1.0);
 }
 
-bool FEvaluationNotify_AlignToGroundInstance::GetTargetTransform(UE::AnimNext::FEvaluationNotifiesTrait::FInstanceData& TraitInstanceData, FTransform& Transform)
+bool FEvaluationNotify_AlignToGroundInstance::GetTargetTransform(UE::UAF::FEvaluationNotifiesTrait::FInstanceData& TraitInstanceData, FTransform& Transform)
 {
 	const UNotifyState_AlignToGround* AlignToGroundNotify = CastChecked<UNotifyState_AlignToGround>(AnimNotify);
 	const FCollisionShape CollisionShape = FCollisionShape::MakeSphere(AlignToGroundNotify->TraceRadius);
@@ -199,7 +199,7 @@ bool FEvaluationNotify_AlignToGroundInstance::GetTargetTransform(UE::AnimNext::F
 	return false;
 }
 
-FTransform GetModelSpaceTransform(const UE::AnimNext::FLODPoseStack& Pose, FBoneIndexType SkeletonBoneIndex)
+FTransform GetModelSpaceTransform(const UE::UAF::FLODPoseStack& Pose, FBoneIndexType SkeletonBoneIndex)
 {
 	const TArrayView<const FBoneIndexType> SkeletonToPoseIndexMap = Pose.GetSkeletonBoneIndexToLODBoneIndexMap();
 	const FBoneIndexType PoseIndex = SkeletonToPoseIndexMap[SkeletonBoneIndex];
@@ -216,7 +216,7 @@ FTransform GetModelSpaceTransform(const UE::AnimNext::FLODPoseStack& Pose, FBone
 	return BoneTransform;
 }
 
-bool FEvaluationNotify_AlignmentInstance::GetTargetTransform(UE::AnimNext::FEvaluationNotifiesTrait::FInstanceData& TraitInstanceData, FTransform& Transform)
+bool FEvaluationNotify_AlignmentInstance::GetTargetTransform(UE::UAF::FEvaluationNotifiesTrait::FInstanceData& TraitInstanceData, FTransform& Transform)
 {
 	const UNotifyState_Alignment* AlignmentNotify = CastChecked<UNotifyState_Alignment>(AnimNotify);
 	// TODO: this should be a proper variable reference held on the notify, but currently this uses a 'named transform' system too. Needs resolving
@@ -225,7 +225,7 @@ bool FEvaluationNotify_AlignmentInstance::GetTargetTransform(UE::AnimNext::FEval
 	PRAGMA_ENABLE_DEPRECATION_WARNINGS
 }
 
-void FEvaluationNotify_AlignmentInstance::Update(UE::AnimNext::FEvaluationNotifiesTrait::FInstanceData& TraitInstanceData, UE::AnimNext::FEvaluationVM& VM)
+void FEvaluationNotify_AlignmentInstance::Update(UE::UAF::FEvaluationNotifiesTrait::FInstanceData& TraitInstanceData, UE::UAF::FEvaluationVM& VM)
 {
 	if (bDisabled)
 	{
@@ -234,7 +234,7 @@ void FEvaluationNotify_AlignmentInstance::Update(UE::AnimNext::FEvaluationNotifi
 	
 	UNotifyState_AlignmentBase* AlignmentNotify = CastChecked<UNotifyState_AlignmentBase>(AnimNotify);
 
-	if (const TUniquePtr<UE::AnimNext::FKeyframeState>* Keyframe = VM.PeekValue<TUniquePtr<UE::AnimNext::FKeyframeState>>(UE::AnimNext::KEYFRAME_STACK_NAME, 0))
+	if (const TUniquePtr<UE::UAF::FKeyframeState>* Keyframe = VM.PeekValue<TUniquePtr<UE::UAF::FKeyframeState>>(UE::UAF::KEYFRAME_STACK_NAME, 0))
 	{
 		if (bFirstFrame)
 		{
@@ -253,9 +253,9 @@ void FEvaluationNotify_AlignmentInstance::Update(UE::AnimNext::FEvaluationNotifi
 			}
 			
 			ActualStartTime = CurrentTime;
-			RoundedEndTime = ActualStartTime + FMath::CeilToDouble((EndTime - ActualStartTime) / AnimNextAlignment::FrameTime) * AnimNextAlignment::FrameTime;
+			RoundedEndTime = ActualStartTime + FMath::CeilToDouble((EndTime - ActualStartTime) / UE::UAF::Alignment::FrameTime) * UE::UAF::Alignment::FrameTime;
 			
-			const UE::AnimNext::FLODPoseStack& Pose = Keyframe->Get()->Pose;
+			const UE::UAF::FLODPoseStack& Pose = Keyframe->Get()->Pose;
 			AlignBone.Initialize(Pose.GetSkeletonAsset());
 
 			if (GetTargetTransform(TraitInstanceData, TargetTransform))
@@ -291,7 +291,7 @@ void FEvaluationNotify_AlignmentInstance::Update(UE::AnimNext::FEvaluationNotifi
 			
 				// extract root motion trajectory todo: this should be cached and reused
 			
-				int NumFrames = 1 + (RoundedEndTime - ActualStartTime) / AnimNextAlignment::FrameTime;
+				int NumFrames = 1 + (RoundedEndTime - ActualStartTime) / UE::UAF::Alignment::FrameTime;
 
 				if (NumFrames <= 0)
 				{
@@ -308,16 +308,16 @@ void FEvaluationNotify_AlignmentInstance::Update(UE::AnimNext::FEvaluationNotifi
 				float PredictionTime = CurrentTime;
 				
 				AnimTrajectoryData.Trajectory[0] = PredictedTransform;
-				PredictionTime += AnimNextAlignment::FrameTime;
+				PredictionTime += UE::UAF::Alignment::FrameTime;
 				AnimTrajectoryData.TranslationCurve[0] = 0;
 
 				for(int i=1;i<NumFrames; i++)
 				{
-					FTransform RootMotionThisFrame = TraitInstanceData.OnExtractRootMotionAttribute.Execute(PredictionTime, AnimNextAlignment::FrameTime, false);
+					FTransform RootMotionThisFrame = TraitInstanceData.OnExtractRootMotionAttribute.Execute(PredictionTime, UE::UAF::Alignment::FrameTime, false);
 				
 					PredictedTransform = RootMotionThisFrame * PredictedTransform;
 					AnimTrajectoryData.Trajectory[i] = PredictedTransform;
-					PredictionTime += AnimNextAlignment::FrameTime;
+					PredictionTime += UE::UAF::Alignment::FrameTime;
 
 					AnimTrajectoryData.TranslationCurve[i] = RootMotionThisFrame.GetTranslation().Length() + AnimTrajectoryData.TranslationCurve[i-1]; 
 					AnimTrajectoryData.RotationCurve[i] = fabs(RootMotionThisFrame.GetRotation().GetAngle()) + AnimTrajectoryData.RotationCurve[i-1]; 
@@ -385,8 +385,8 @@ void FEvaluationNotify_AlignmentInstance::Update(UE::AnimNext::FEvaluationNotifi
 					
 					if (AlignmentNotify->bSeparateTranslationCurves)
 					{
-						float InMovementDirectionWeight = GetWeight(ActualStartTime + AnimNextAlignment::FrameTime*i, AlignmentNotify->TranslationWarpingCurve_InMovementDirection);
-						float OutOfMovementDirectionWeight = GetWeight(ActualStartTime + AnimNextAlignment::FrameTime*i, AlignmentNotify->TranslationWarpingCurve_OutOfMovementDirection);
+						float InMovementDirectionWeight = GetWeight(ActualStartTime + UE::UAF::Alignment::FrameTime*i, AlignmentNotify->TranslationWarpingCurve_InMovementDirection);
+						float OutOfMovementDirectionWeight = GetWeight(ActualStartTime + UE::UAF::Alignment::FrameTime*i, AlignmentNotify->TranslationWarpingCurve_OutOfMovementDirection);
 	
 						FVector Delta = TransformFromTarget.GetTranslation() - OldPosition;
 						FVector InMovementDirectionDelta = MovementDirection * Delta.Dot(MovementDirection);
@@ -396,7 +396,7 @@ void FEvaluationNotify_AlignmentInstance::Update(UE::AnimNext::FEvaluationNotifi
 					}
 					else
 					{
-						float Weight = GetWeight(ActualStartTime + AnimNextAlignment::FrameTime*i, AlignmentNotify->TranslationWarpingCurve);
+						float Weight = GetWeight(ActualStartTime + UE::UAF::Alignment::FrameTime*i, AlignmentNotify->TranslationWarpingCurve);
 						NewPosition = FMath::Lerp(OldPosition, TransformFromTarget.GetTranslation(), Weight);
 					}
 
@@ -437,7 +437,7 @@ void FEvaluationNotify_AlignmentInstance::Update(UE::AnimNext::FEvaluationNotifi
 				// Rotation warping
 				for(int i=0;i<NumFrames; i++)
 				{
-					float Weight = GetWeight(ActualStartTime + AnimNextAlignment::FrameTime*i, AlignmentNotify->RotationWarpingCurve);
+					float Weight = GetWeight(ActualStartTime + UE::UAF::Alignment::FrameTime*i, AlignmentNotify->RotationWarpingCurve);
 				
 					FQuat OldRotation = WarpedTrajectory[i].GetRotation();
 					FTransform TransformFromTarget = AnimTrajectoryData.Trajectory[i] * InverseLastFrame * TargetTransform;
@@ -453,7 +453,7 @@ void FEvaluationNotify_AlignmentInstance::Update(UE::AnimNext::FEvaluationNotifi
     		
 		if (RootMotionProvider && !WarpedTrajectory.IsEmpty())
 		{
-			float Frame = (CurrentTime + TraitInstanceData.DeltaTime - ActualStartTime) / AnimNextAlignment::FrameTime;
+			float Frame = (CurrentTime + TraitInstanceData.DeltaTime - ActualStartTime) / UE::UAF::Alignment::FrameTime;
 			
 			if (PreviousFrame < WarpedTrajectory.Num() - 1)
 			{
@@ -461,14 +461,14 @@ void FEvaluationNotify_AlignmentInstance::Update(UE::AnimNext::FEvaluationNotifi
 				
 				if (Frame > WarpedTrajectory.Num()-1)
 				{
-					float PostAlignmentTime = (Frame - (WarpedTrajectory.Num() - 1)) * AnimNextAlignment::FrameTime;
+					float PostAlignmentTime = (Frame - (WarpedTrajectory.Num() - 1)) * UE::UAF::Alignment::FrameTime;
 					FTransform PostAlignmentRootMotion = TraitInstanceData.OnExtractRootMotionAttribute.Execute(RoundedEndTime, PostAlignmentTime, false);
-					AnimNextAlignment::GetTransformForFrame(Frame, WarpedTrajectory, WorldTransform);
+					UE::UAF::Alignment::GetTransformForFrame(Frame, WarpedTrajectory, WorldTransform);
 					WorldTransform = PostAlignmentRootMotion * TargetTransform;
 				}
 				else
 				{
-					AnimNextAlignment::GetTransformForFrame(Frame, WarpedTrajectory, WorldTransform);
+					UE::UAF::Alignment::GetTransformForFrame(Frame, WarpedTrajectory, WorldTransform);
 				}
 
 				if(AlignmentNotify->UpdateMode == EAlignmentUpdateMode::World)
@@ -478,7 +478,7 @@ void FEvaluationNotify_AlignmentInstance::Update(UE::AnimNext::FEvaluationNotifi
 				else // relative mode
 				{
 					FTransform PrevTransform;
-					AnimNextAlignment::GetTransformForFrame(PreviousFrame, WarpedTrajectory, PrevTransform);
+					UE::UAF::Alignment::GetTransformForFrame(PreviousFrame, WarpedTrajectory, PrevTransform);
 					RootMotionProvider->OverrideRootMotion(WorldTransform.GetRelativeTransform(PrevTransform), Keyframe->Get()->Attributes);
 				}
 		
