Fix math errors with alignment bone feature in Alignment Notify

diff --git a/Engine/Plugins/Experimental/EvaluationNotifies/Source/Runtime/Private/AnimNext/AnimNextAlignment.cpp b/Engine/Plugins/Experimental/EvaluationNotifies/Source/Runtime/Private/AnimNext/AnimNextAlignment.cpp
--- a/Engine/Plugins/Experimental/EvaluationNotifies/Source/Runtime/Private/AnimNext/AnimNextAlignment.cpp
+++ b/Engine/Plugins/Experimental/EvaluationNotifies/Source/Runtime/Private/AnimNext/AnimNextAlignment.cpp
@@ -265,6 +265,8 @@ void FEvaluationNotify_AlignmentInstance::Update(UE::UAF::FEvaluationNotifiesTra
 			
 				// get alignment transform
 				const float PredictionDelta = RoundedEndTime - CurrentTime;
+				
+				TargetTransform = AlignmentNotify->AlignOffset * TargetTransform;
 
 				if (AlignBone.HasValidSetup())
 				{
@@ -274,10 +276,8 @@ void FEvaluationNotify_AlignmentInstance::Update(UE::UAF::FEvaluationNotifiesTra
 					FTransform PredictedRootMotion = TraitInstanceData.OnExtractRootMotionAttribute.Execute(CurrentTime, PredictionDelta, false);
 
 					FTransform AlignBoneTransform = GetModelSpaceTransform(Pose, AlignBone.BoneIndex);
-				 	TargetTransform = AlignBoneTransform.Inverse() * PredictedRootMotion * TargetTransform;
+				 	TargetTransform = PredictedRootMotion * AlignBoneTransform.Inverse() * TargetTransform;
 				 }
-			
-				TargetTransform = AlignmentNotify->AlignOffset * TargetTransform;
 
 				if (AlignmentNotify->bForceTargetTransformUpright)
 				{
@@ -464,7 +464,7 @@ void FEvaluationNotify_AlignmentInstance::Update(UE::UAF::FEvaluationNotifiesTra
 					float PostAlignmentTime = (Frame - (WarpedTrajectory.Num() - 1)) * UE::UAF::Alignment::FrameTime;
 					FTransform PostAlignmentRootMotion = TraitInstanceData.OnExtractRootMotionAttribute.Execute(RoundedEndTime, PostAlignmentTime, false);
 					UE::UAF::Alignment::GetTransformForFrame(Frame, WarpedTrajectory, WorldTransform);
-					WorldTransform = PostAlignmentRootMotion * TargetTransform;
+					WorldTransform = PostAlignmentRootMotion * WarpedTrajectory.Last();
 				}
 				else
 				{
