EvaluationNotifies - fix for missing UE::UAF::FEvaluationNotifiesTrait::UnregisterEvaluationHandler(UNotifyState_AlignToGround::StaticClass())

#Synced-CL 48365056

diff --git a/Engine/Plugins/Experimental/EvaluationNotifies/Source/Runtime/Private/EvaluationNotifiesRuntimeModule.cpp b/Engine/Plugins/Experimental/EvaluationNotifies/Source/Runtime/Private/EvaluationNotifiesRuntimeModule.cpp
--- a/Engine/Plugins/Experimental/EvaluationNotifies/Source/Runtime/Private/EvaluationNotifiesRuntimeModule.cpp
+++ b/Engine/Plugins/Experimental/EvaluationNotifies/Source/Runtime/Private/EvaluationNotifiesRuntimeModule.cpp
@@ -47,6 +47,7 @@ namespace UE
 			FAnimNode_EvaluationNotifies::UnregisterEvaluationHandler(UNotifyState_TwoBoneIK::StaticClass());
 			
 			UE::UAF::FEvaluationNotifiesTrait::UnregisterEvaluationHandler(UNotifyState_Alignment::StaticClass());
+			UE::UAF::FEvaluationNotifiesTrait::UnregisterEvaluationHandler(UNotifyState_AlignToGround::StaticClass());
 		}
 	}
 }
